UAF: Moved/renamed AnimNext and AnimNextAnimGraph plugins

diff --git a/Engine/Plugins/Experimental/UAF/UAFAnimGraph/Config/BaseUAFAnimGraph.ini b/Engine/Plugins/Experimental/UAF/UAFAnimGraph/Config/BaseUAFAnimGraph.ini
new file mode 100644
--- /dev/null
+++ b/Engine/Plugins/Experimental/UAF/UAFAnimGraph/Config/BaseUAFAnimGraph.ini
@@ -0,0 +1,205 @@
+[CoreRedirects]
+; Traits rename
++StructRedirects=(OldName="/Script/AnimNext.AnimNextBlendByBoolDecoratorSharedData", NewName="/Script/UAFAnimGraph.AnimNextBlendByBoolTraitSharedData")
++StructRedirects=(OldName="/Script/AnimNext.AnimNextBlendInertializerDecoratorSharedData", NewName="/Script/UAFAnimGraph.AnimNextBlendInertializerTraitSharedData")
++StructRedirects=(OldName="/Script/AnimNext.AnimNextBlendSmootherDecoratorSharedData", NewName="/Script/UAFAnimGraph.AnimNextBlendSmootherTraitSharedData")
++StructRedirects=(OldName="/Script/AnimNext.AnimNextBlendSmootherPerBoneDecoratorSharedData", NewName="/Script/UAFAnimGraph.AnimNextBlendProfilePerChildProviderTraitSharedData")
++StructRedirects=(OldName="/Script/AnimNext.AnimNextBlendTwoWayDecoratorSharedData", NewName="/Script/UAFAnimGraph.AnimNextBlendTwoWayTraitSharedData")
++StructRedirects=(OldName="/Script/AnimNext.AnimNextReferencePoseDecoratorSharedData", NewName="/Script/UAFAnimGraph.AnimNextReferencePoseTraitSharedData")
++StructRedirects=(OldName="/Script/AnimNext.AnimNextSequencePlayerDecoratorSharedData", NewName="/Script/UAFAnimGraph.AnimNextSequencePlayerTraitSharedData")
++StructRedirects=(OldName="/Script/AnimNext.AnimNextSubGraphHostDecoratorSharedData", NewName="/Script/UAFAnimGraph.AnimNextSubGraphHostTraitSharedData")
++StructRedirects=(OldName="/Script/AnimNext.AnimNextSynchronizeUsingGroupsDecoratorSharedData", NewName="/Script/UAFAnimGraph.AnimNextSynchronizeUsingGroupsTraitSharedData")
+
+; Plugin refactor
++ClassRedirects=(OldName="/Script/AnimNextUncookedOnly.AnimGraphNode_AnimNextGraph", NewName="/Script/UAFAnimGraphUncookedOnly.AnimGraphNode_AnimNextGraph")
+
++StructRedirects=(OldName="/Script/AnimNext.AnimNextBlendByBoolTraitSharedData", NewName="/Script/UAFAnimGraph.AnimNextBlendByBoolTraitSharedData")
++StructRedirects=(OldName="/Script/AnimNext.AnimNextBlendInertializerTraitSharedData", NewName="/Script/UAFAnimGraph.AnimNextBlendInertializerTraitSharedData")
++StructRedirects=(OldName="/Script/AnimNext.AnimNextBlendSmootherTraitSharedData", NewName="/Script/UAFAnimGraph.AnimNextBlendSmootherTraitSharedData")
++StructRedirects=(OldName="/Script/AnimNext.AnimNextBlendSmootherPerBoneTraitSharedData", NewName="/Script/UAFAnimGraph.AnimNextBlendProfilePerChildProviderTraitSharedData")
++StructRedirects=(OldName="/Script/AnimNext.AnimNextBlendTwoWayTraitSharedData", NewName="/Script/UAFAnimGraph.AnimNextBlendTwoWayTraitSharedData")
++StructRedirects=(OldName="/Script/AnimNext.AnimNextReferencePoseTraitSharedData", NewName="/Script/UAFAnimGraph.AnimNextReferencePoseTraitSharedData")
++StructRedirects=(OldName="/Script/AnimNext.AnimNextSequencePlayerTraitSharedData", NewName="/Script/UAFAnimGraph.AnimNextSequencePlayerTraitSharedData")
++StructRedirects=(OldName="/Script/AnimNext.AnimNextSubGraphHostTraitSharedData", NewName="/Script/UAFAnimGraph.AnimNextSubGraphHostTraitSharedData")
++StructRedirects=(OldName="/Script/AnimNext.AnimNextSynchronizeUsingGroupsTraitSharedData", NewName="/Script/UAFAnimGraph.AnimNextSynchronizeUsingGroupsTraitSharedData")
++StructRedirects=(OldName="/Script/AnimNext.AnimNode_AnimNextGraph", NewName="/Script/UAFAnimGraph.AnimNode_AnimNextGraph")
+
+; AnimNextGraph -> AnimNextModule
++PropertyRedirects=(OldName="/Script/AnimNextGraph.AnimNode_AnimNextGraph.AnimNextGraph", NewName="Module")
+
+; Rename K2 Node
++ClassRedirects=(OldName="/Script/AnimNextAnimGraphUncookedOnly.K2Node_PlayAnim", NewName="/Script/UAFAnimGraphUncookedOnly.K2Node_AnimNextPlayAnim")
+
+; Rename PlayAnim parameter
++PropertyRedirects=(OldName="/Script/AnimNextAnimGraphUncookedOnly.CreateProxyObjectForPlayAnim.AnimationObject", NewName="AnimSequence")
++PropertyRedirects=(OldName="/Script/AnimNextAnimGraphUncookedOnly.K2Node_PlayAnim.AnimationObject", NewName="AnimSequence")
+
+; PlayAnim -> Injection
++EnumRedirects=(OldName="/Script/AnimNextAnimGraph.EAnimNextPlayAnimBlendMode", NewName="/Script/UAFAnimGraph.EAnimNextInjectionBlendMode")
++EnumRedirects=(OldName="/Script/AnimNextAnimGraph.EAnimNextPlayAnimStatus", NewName="/Script/UAFAnimGraph.EAnimNextInjectionStatus")
++StructRedirects=(OldName="/Script/AnimNextAnimGraph.AnimNextPlayAnimBlendSettings", NewName="/Script/UAFAnimGraph.AnimNextInjectionBlendSettings")
++StructRedirects=(OldName="/Script/AnimNextAnimGraph.AnimNextPlayAnimRequestArgs", NewName="/Script/UAFAnimGraph.AnimNextInjectionRequestArgs")
++PropertyRedirects=(OldName="/Script/AnimNextAnimGraph.AnimNextInjectionRequestArgs.SlotName", NewName="Name")
++ClassRedirects=(OldName="/Script/AnimNextAnimGraphUncookedOnly.K2Node_AnimNextPlayAsset", NewName="/Script/UAFAnimGraphUncookedOnly.K2Node_AnimNextInjection")
++StructRedirects=(OldName="/Script/AnimNextAnimGraph.RigUnit_AnimNextRunAnimationGraph", NewName="/Script/UAFAnimGraph.RigUnit_AnimNextRunAnimationGraph_v1")
++StructRedirects=(OldName="/Script/AnimNextAnimGraph.AnimNextGraphInstancePtr", NewName="/Script/UAFAnimGraph.AnimNextAnimGraph")
+
+; AnimNextAnimGraph plugin -> UAFAnimGraph plugin
++PackageRedirects=(OldName="/Script/AnimNextAnimGraph", NewName="/Script/UAFAnimGraph")
++PackageRedirects=(OldName="/Script/AnimNextAnimGraphEditor", NewName="/Script/UAFAnimGraphEditor")
++PackageRedirects=(OldName="/Script/AnimNextAnimGraphUncookedOnly", NewName="/Script/UAFAnimGraphUncookedOnly")
++PackageRedirects=(OldName="/Script/AnimNextAnimGraphTestSuite", NewName="/Script/UAFAnimGraphTestSuite")
+
++StructRedirects=(OldName="/Script/AnimNextAnimGraph.AnimNode_AnimNextGraph", NewName="/Script/UAFAnimGraph.AnimNode_AnimNextGraph")
++StructRedirects=(OldName="/Script/AnimNextAnimGraph.AnimNextGraphContextData", NewName="/Script/UAFAnimGraph.AnimNextGraphContextData")
++StructRedirects=(OldName="/Script/AnimNextAnimGraph.AnimNextGraphLatentPropertiesContextData", NewName="/Script/UAFAnimGraph.AnimNextGraphLatentPropertiesContextData")
++StructRedirects=(OldName="/Script/AnimNextAnimGraph.AnimNextModuleAnimGraphComponent", NewName="/Script/UAFAnimGraph.AnimNextModuleAnimGraphComponent")
++StructRedirects=(OldName="/Script/AnimNextAnimGraph.RigDecorator_AnimNextCppDecorator", NewName="/Script/UAFAnimGraph.RigDecorator_AnimNextCppDecorator")
++StructRedirects=(OldName="/Script/AnimNextAnimGraph.RigUnit_AnimNextGraphRoot", NewName="/Script/UAFAnimGraph.RigUnit_AnimNextGraphRoot")
++StructRedirects=(OldName="/Script/AnimNextAnimGraph.AnimNextRunGraphWorkData_v1", NewName="/Script/UAFAnimGraph.AnimNextRunGraphWorkData_v1")
++StructRedirects=(OldName="/Script/AnimNextAnimGraph.RigUnit_AnimNextRunAnimationGraph_v1", NewName="/Script/UAFAnimGraph.RigUnit_AnimNextRunAnimationGraph_v1")
++StructRedirects=(OldName="/Script/AnimNextAnimGraph.AnimNextRunGraphWorkData_v2", NewName="/Script/UAFAnimGraph.AnimNextRunGraphWorkData_v2")
++StructRedirects=(OldName="/Script/AnimNextAnimGraph.RigUnit_AnimNextShimRoot", NewName="/Script/UAFAnimGraph.RigUnit_AnimNextShimRoot")
++StructRedirects=(OldName="/Script/AnimNextAnimGraph.RigUnit_AnimNextTraitStack", NewName="/Script/UAFAnimGraph.RigUnit_AnimNextTraitStack")
++StructRedirects=(OldName="/Script/AnimNextAnimGraph.AnimNextBlendStackCoreTraitSharedData", NewName="/Script/UAFAnimGraph.AnimNextBlendStackCoreTraitSharedData")
++StructRedirects=(OldName="/Script/AnimNextAnimGraph.AnimNextBlendStackTraitSharedData", NewName="/Script/UAFAnimGraph.AnimNextBlendStackTraitSharedData")
++StructRedirects=(OldName="/Script/AnimNextAnimGraph.AnimNextBlendStackRequesterTraitSharedData", NewName="/Script/UAFAnimGraph.AnimNextBlendStackRequesterTraitSharedData")
++StructRedirects=(OldName="/Script/AnimNextAnimGraph.AnimNextCallFunctionSharedData", NewName="/Script/UAFAnimGraph.AnimNextCallFunctionSharedData")
++StructRedirects=(OldName="/Script/AnimNextAnimGraph.GCGraphInstanceComponent", NewName="/Script/UAFAnimGraph.GCGraphInstanceComponent")
++StructRedirects=(OldName="/Script/AnimNextAnimGraph.RigUnit_AnimNextSetNotifyContext", NewName="/Script/UAFAnimGraph.RigUnit_AnimNextSetNotifyContext")
++StructRedirects=(OldName="/Script/AnimNextAnimGraph.RigUnit_GetPostProcessAnimation", NewName="/Script/UAFAnimGraph.RigUnit_GetPostProcessAnimation")
++StructRedirects=(OldName="/Script/AnimNextAnimGraph.RigUnit_HasValidEvaluationModifier", NewName="/Script/UAFAnimGraph.RigUnit_HasValidEvaluationModifier")
++StructRedirects=(OldName="/Script/AnimNextAnimGraph.SyncGroupGraphInstanceComponent", NewName="/Script/UAFAnimGraph.SyncGroupGraphInstanceComponent")
++StructRedirects=(OldName="/Script/AnimNextAnimGraph.AnimNextNativeDataInterface_AnimSequencePlayer", NewName="/Script/UAFAnimGraph.AnimNextNativeDataInterface_AnimSequencePlayer")
++StructRedirects=(OldName="/Script/AnimNextAnimGraph.AnimNextNativeDataInterface_BlendSpacePlayer", NewName="/Script/UAFAnimGraph.AnimNextNativeDataInterface_BlendSpacePlayer")
++StructRedirects=(OldName="/Script/AnimNextAnimGraph.AnimNextNativeDataInterface_SynchronizeUsingGroups", NewName="/Script/UAFAnimGraph.AnimNextNativeDataInterface_SynchronizeUsingGroups")
++StructRedirects=(OldName="/Script/AnimNextAnimGraph.AnimNextModuleInjectionComponent", NewName="/Script/UAFAnimGraph.AnimNextModuleInjectionComponent")
++StructRedirects=(OldName="/Script/AnimNextAnimGraph.GraphInstanceInjectionComponent", NewName="/Script/UAFAnimGraph.GraphInstanceInjectionComponent")
++StructRedirects=(OldName="/Script/AnimNextAnimGraph.AnimNextInjectionSiteTraitSharedData", NewName="/Script/UAFAnimGraph.AnimNextInjectionSiteTraitSharedData")
++StructRedirects=(OldName="/Script/AnimNextAnimGraph.AnimNextPlayAnimSlotTraitSharedData", NewName="/Script/UAFAnimGraph.AnimNextPlayAnimSlotTraitSharedData")
++StructRedirects=(OldName="/Script/AnimNextAnimGraph.AnimNextApplyAdditiveTraitSharedData", NewName="/Script/UAFAnimGraph.AnimNextApplyAdditiveTraitSharedData")
++StructRedirects=(OldName="/Script/AnimNextAnimGraph.BoolAlphaInputTraitArgs", NewName="/Script/UAFAnimGraph.BoolAlphaInputTraitArgs")
++StructRedirects=(OldName="/Script/AnimNextAnimGraph.CurveAlphaInputTraitArgs", NewName="/Script/UAFAnimGraph.CurveAlphaInputTraitArgs")
++StructRedirects=(OldName="/Script/AnimNextAnimGraph.FloatAlphaInputTraitArgs", NewName="/Script/UAFAnimGraph.FloatAlphaInputTraitArgs")
++StructRedirects=(OldName="/Script/AnimNextAnimGraph.SimpleFloatAlphaInputTraitArgs", NewName="/Script/UAFAnimGraph.SimpleFloatAlphaInputTraitArgs")
++StructRedirects=(OldName="/Script/AnimNextAnimGraph.AnimNextBlendByBoolTraitSharedData", NewName="/Script/UAFAnimGraph.AnimNextBlendByBoolTraitSharedData")
++StructRedirects=(OldName="/Script/AnimNextAnimGraph.AnimNextBlendInertializerCoreTraitSharedData", NewName="/Script/UAFAnimGraph.AnimNextBlendInertializerCoreTraitSharedData")
++StructRedirects=(OldName="/Script/AnimNextAnimGraph.AnimNextBlendInertializerTraitSharedData", NewName="/Script/UAFAnimGraph.AnimNextBlendInertializerTraitSharedData")
++StructRedirects=(OldName="/Script/AnimNextAnimGraph.AnimNextBlendLayerTraitSharedData", NewName="/Script/UAFAnimGraph.AnimNextBlendLayerTraitSharedData")
++StructRedirects=(OldName="/Script/AnimNextAnimGraph.AnimNextBlendSmootherPerBoneCoreTraitSharedData", NewName="/Script/UAFAnimGraph.AnimNextBlendSmootherPerBoneCoreTraitSharedData")
++StructRedirects=(OldName="/Script/AnimNextAnimGraph.AnimNextBlendProfilePerChildProviderTraitSharedData", NewName="/Script/UAFAnimGraph.AnimNextBlendProfilePerChildProviderTraitSharedData")
++StructRedirects=(OldName="/Script/AnimNextAnimGraph.AnimNextBlendProfileProviderTraitSharedData", NewName="/Script/UAFAnimGraph.AnimNextBlendProfileProviderTraitSharedData")
++StructRedirects=(OldName="/Script/AnimNextAnimGraph.AnimNextBlendTwoWayTraitSharedData", NewName="/Script/UAFAnimGraph.AnimNextBlendTwoWayTraitSharedData")
++StructRedirects=(OldName="/Script/AnimNextAnimGraph.AnimNextInertializationDeadBlendTraitSharedData", NewName="/Script/UAFAnimGraph.AnimNextInertializationDeadBlendTraitSharedData")
++StructRedirects=(OldName="/Script/AnimNextAnimGraph.AnimNextInputPoseTraitSharedData", NewName="/Script/UAFAnimGraph.AnimNextInputPoseTraitSharedData")
++StructRedirects=(OldName="/Script/AnimNextAnimGraph.ModifyCurveParameters", NewName="/Script/UAFAnimGraph.ModifyCurveParameters")
++StructRedirects=(OldName="/Script/AnimNextAnimGraph.AnimNextModifyCurveTraitSharedData", NewName="/Script/UAFAnimGraph.AnimNextModifyCurveTraitSharedData")
++StructRedirects=(OldName="/Script/AnimNextAnimGraph.ModifyCurveTask", NewName="/Script/UAFAnimGraph.ModifyCurveTask")
++StructRedirects=(OldName="/Script/AnimNextAnimGraph.AnimNextNotifyDispatcherComponent", NewName="/Script/UAFAnimGraph.AnimNextNotifyDispatcherComponent")
++StructRedirects=(OldName="/Script/AnimNextAnimGraph.AnimNextNotifyFilterTraitSharedData", NewName="/Script/UAFAnimGraph.AnimNextNotifyFilterTraitSharedData")
++StructRedirects=(OldName="/Script/AnimNextAnimGraph.AnimNextPassthroughSharedData", NewName="/Script/UAFAnimGraph.AnimNextPassthroughSharedData")
++StructRedirects=(OldName="/Script/AnimNextAnimGraph.AnimNextReferencePoseTraitSharedData", NewName="/Script/UAFAnimGraph.AnimNextReferencePoseTraitSharedData")
++StructRedirects=(OldName="/Script/AnimNextAnimGraph.AnimNextSubGraphHostTraitSharedData", NewName="/Script/UAFAnimGraph.AnimNextSubGraphHostTraitSharedData")
++StructRedirects=(OldName="/Script/AnimNextAnimGraph.AnimNextEvaluationTask", NewName="/Script/UAFAnimGraph.AnimNextEvaluationTask")
++StructRedirects=(OldName="/Script/AnimNextAnimGraph.SerializableEvaluationProgram", NewName="/Script/UAFAnimGraph.SerializableEvaluationProgram")
++StructRedirects=(OldName="/Script/AnimNextAnimGraph.AnimNextApplyAdditiveKeyframeTask", NewName="/Script/UAFAnimGraph.AnimNextApplyAdditiveKeyframeTask")
++StructRedirects=(OldName="/Script/AnimNextAnimGraph.AnimNextBlendTwoKeyframesTask", NewName="/Script/UAFAnimGraph.AnimNextBlendTwoKeyframesTask")
++StructRedirects=(OldName="/Script/AnimNextAnimGraph.AnimNextBlendOverwriteKeyframeWithScaleTask", NewName="/Script/UAFAnimGraph.AnimNextBlendOverwriteKeyframeWithScaleTask")
++StructRedirects=(OldName="/Script/AnimNextAnimGraph.AnimNextBlendAddKeyframeWithScaleTask", NewName="/Script/UAFAnimGraph.AnimNextBlendAddKeyframeWithScaleTask")
++StructRedirects=(OldName="/Script/AnimNextAnimGraph.AnimNextBlendOverwriteKeyframePerBoneWithScaleTask", NewName="/Script/UAFAnimGraph.AnimNextBlendOverwriteKeyframePerBoneWithScaleTask")
++StructRedirects=(OldName="/Script/AnimNextAnimGraph.AnimNextBlendAddKeyframePerBoneWithScaleTask", NewName="/Script/UAFAnimGraph.AnimNextBlendAddKeyframePerBoneWithScaleTask")
++StructRedirects=(OldName="/Script/AnimNextAnimGraph.AnimNextBlendKeyframePerBoneWithScaleTask", NewName="/Script/UAFAnimGraph.AnimNextBlendKeyframePerBoneWithScaleTask")
++StructRedirects=(OldName="/Script/AnimNextAnimGraph.AnimNextConvertRotationsToLocalSpaceTask", NewName="/Script/UAFAnimGraph.AnimNextConvertRotationsToLocalSpaceTask")
++StructRedirects=(OldName="/Script/AnimNextAnimGraph.AnimNextConvertRotationsToMeshSpaceTask", NewName="/Script/UAFAnimGraph.AnimNextConvertRotationsToMeshSpaceTask")
++StructRedirects=(OldName="/Script/AnimNextAnimGraph.AnimNextDeadBlendingTransitionTask", NewName="/Script/UAFAnimGraph.AnimNextDeadBlendingTransitionTask")
++StructRedirects=(OldName="/Script/AnimNextAnimGraph.AnimNextDeadBlendingApplyTask", NewName="/Script/UAFAnimGraph.AnimNextDeadBlendingApplyTask")
++StructRedirects=(OldName="/Script/AnimNextAnimGraph.AnimNextExecuteProgramTask", NewName="/Script/UAFAnimGraph.AnimNextExecuteProgramTask")
++StructRedirects=(OldName="/Script/AnimNextAnimGraph.AnimNextNormalizeKeyframeRotationsTask", NewName="/Script/UAFAnimGraph.AnimNextNormalizeKeyframeRotationsTask")
++StructRedirects=(OldName="/Script/AnimNextAnimGraph.AnimNextAnimSequenceKeyframeTask", NewName="/Script/UAFAnimGraph.AnimNextAnimSequenceKeyframeTask")
++StructRedirects=(OldName="/Script/AnimNextAnimGraph.AnimNextPushPoseTask", NewName="/Script/UAFAnimGraph.AnimNextPushPoseTask")
++StructRedirects=(OldName="/Script/AnimNextAnimGraph.AnimNextPushReferenceKeyframeTask", NewName="/Script/UAFAnimGraph.AnimNextPushReferenceKeyframeTask")
++StructRedirects=(OldName="/Script/AnimNextAnimGraph.AnimNextSwapTransformsTask", NewName="/Script/UAFAnimGraph.AnimNextSwapTransformsTask")
++StructRedirects=(OldName="/Script/AnimNextAnimGraph.AnimNextStoreKeyframeTransformsTask", NewName="/Script/UAFAnimGraph.AnimNextStoreKeyframeTransformsTask")
++StructRedirects=(OldName="/Script/AnimNextAnimGraph.AnimNextDuplicateTopKeyframeTask", NewName="/Script/UAFAnimGraph.AnimNextDuplicateTopKeyframeTask")
++StructRedirects=(OldName="/Script/AnimNextAnimGraph.AnimNextAnimGraphBuilder", NewName="/Script/UAFAnimGraph.AnimNextAnimGraphBuilder")
++StructRedirects=(OldName="/Script/AnimNextAnimGraph.AnimNextFactoryParams", NewName="/Script/UAFAnimGraph.AnimNextFactoryParams")
++StructRedirects=(OldName="/Script/AnimNextAnimGraph.AnimNextSimpleAnimGraphBuilderTraitStackDesc", NewName="/Script/UAFAnimGraph.AnimNextSimpleAnimGraphBuilderTraitStackDesc")
++StructRedirects=(OldName="/Script/AnimNextAnimGraph.AnimNextSimpleAnimGraphBuilder", NewName="/Script/UAFAnimGraph.AnimNextSimpleAnimGraphBuilder")
++StructRedirects=(OldName="/Script/AnimNextAnimGraph.AnimNextGraphInjectionData", NewName="/Script/UAFAnimGraph.AnimNextGraphInjectionData")
++StructRedirects=(OldName="/Script/AnimNextAnimGraph.AnimNextAnimGraph", NewName="/Script/UAFAnimGraph.AnimNextAnimGraph")
++StructRedirects=(OldName="/Script/AnimNextAnimGraph.AnimNextGraphEntryPoint", NewName="/Script/UAFAnimGraph.AnimNextGraphEntryPoint")
++StructRedirects=(OldName="/Script/AnimNextAnimGraph.AnimNextGraphEvaluatorExecuteArgument", NewName="/Script/UAFAnimGraph.AnimNextGraphEvaluatorExecuteArgument")
++StructRedirects=(OldName="/Script/AnimNextAnimGraph.AnimNextGraphEvaluatorExecuteDefinition", NewName="/Script/UAFAnimGraph.AnimNextGraphEvaluatorExecuteDefinition")
++StructRedirects=(OldName="/Script/AnimNextAnimGraph.AnimNextGraphInstance", NewName="/Script/UAFAnimGraph.AnimNextGraphInstance")
++StructRedirects=(OldName="/Script/AnimNextAnimGraph.UAFGraphInstanceComponent", NewName="/Script/UAFAnimGraph.UAFGraphInstanceComponent")
++StructRedirects=(OldName="/Script/AnimNextAnimGraph.AnimNextInjectionBlendSettings", NewName="/Script/UAFAnimGraph.AnimNextInjectionBlendSettings")
++StructRedirects=(OldName="/Script/AnimNextAnimGraph.AnimNextInjectionRequestArgs", NewName="/Script/UAFAnimGraph.AnimNextInjectionRequestArgs")
++StructRedirects=(OldName="/Script/AnimNextAnimGraph.AnimNextInjectionSite", NewName="/Script/UAFAnimGraph.AnimNextInjectionSite")
++StructRedirects=(OldName="/Script/AnimNextAnimGraph.AnimNextEntryPointHandle", NewName="/Script/UAFAnimGraph.AnimNextEntryPointHandle")
++StructRedirects=(OldName="/Script/AnimNextAnimGraph.AnimNextTraitHandle", NewName="/Script/UAFAnimGraph.AnimNextTraitHandle")
++StructRedirects=(OldName="/Script/AnimNextAnimGraph.AnimNextTraitSharedData", NewName="/Script/UAFAnimGraph.AnimNextTraitSharedData")
++StructRedirects=(OldName="/Script/AnimNextAnimGraph.AlphaInputTraitArgs", NewName="/Script/UAFAnimGraph.AlphaInputTraitArgs")
++StructRedirects=(OldName="/Script/AnimNextAnimGraph.AnimNextBlendSmootherCoreTraitSharedData", NewName="/Script/UAFAnimGraph.AnimNextBlendSmootherCoreTraitSharedData")
++StructRedirects=(OldName="/Script/AnimNextAnimGraph.AnimNextBlendSmootherTraitSharedData", NewName="/Script/UAFAnimGraph.AnimNextBlendSmootherTraitSharedData")
++StructRedirects=(OldName="/Script/AnimNextAnimGraph.AnimNextBlendSpacePlayerTraitSharedData", NewName="/Script/UAFAnimGraph.AnimNextBlendSpacePlayerTraitSharedData")
++StructRedirects=(OldName="/Script/AnimNextAnimGraph.AnimNextInertializationRequestEvent", NewName="/Script/UAFAnimGraph.AnimNextInertializationRequestEvent")
++StructRedirects=(OldName="/Script/AnimNextAnimGraph.AnimNextNotifyDispatcherTraitSharedData", NewName="/Script/UAFAnimGraph.AnimNextNotifyDispatcherTraitSharedData")
++StructRedirects=(OldName="/Script/AnimNextAnimGraph.PassthroughBlendTraitSharedData", NewName="/Script/UAFAnimGraph.PassthroughBlendTraitSharedData")
++StructRedirects=(OldName="/Script/AnimNextAnimGraph.AnimNextSequencePlayerTraitSharedData", NewName="/Script/UAFAnimGraph.AnimNextSequencePlayerTraitSharedData")
++StructRedirects=(OldName="/Script/AnimNextAnimGraph.AnimNextSynchronizeUsingGroupsTraitSharedData", NewName="/Script/UAFAnimGraph.AnimNextSynchronizeUsingGroupsTraitSharedData")
+
++ClassRedirects=(OldName="/Script/AnimNextAnimGraph.AnimNextAnimGraphSettings", NewName="/Script/UAFAnimGraph.AnimNextAnimGraphSettings")
++ClassRedirects=(OldName="/Script/AnimNextAnimGraph.PlayAnimCallbackProxy", NewName="/Script/UAFAnimGraph.PlayAnimCallbackProxy")
++ClassRedirects=(OldName="/Script/AnimNextAnimGraph.AnimNextSkeletalMeshComponentLibrary", NewName="/Script/UAFAnimGraph.AnimNextSkeletalMeshComponentLibrary")
++ClassRedirects=(OldName="/Script/AnimNextAnimGraph.AnimNextAnimationGraph", NewName="/Script/UAFAnimGraph.AnimNextAnimationGraph")
++ClassRedirects=(OldName="/Script/AnimNextAnimGraph.PostProcessAnimationUserAssetData", NewName="/Script/UAFAnimGraph.PostProcessAnimationUserAssetData")
+
++EnumRedirects=(OldName="/Script/AnimNextAnimGraph.EUninjectionResult", NewName="/Script/UAFAnimGraph.EUninjectionResult")
++EnumRedirects=(OldName="/Script/AnimNextAnimGraph.EAnimNextCallFunctionCallSite", NewName="/Script/UAFAnimGraph.EAnimNextCallFunctionCallSite")
++EnumRedirects=(OldName="/Script/AnimNextAnimGraph.EAnimNext_ModifyCurveApplyMode", NewName="/Script/UAFAnimGraph.EAnimNext_ModifyCurveApplyMode")
++EnumRedirects=(OldName="/Script/AnimNextAnimGraph.EAnimNextReferencePoseType", NewName="/Script/UAFAnimGraph.EAnimNextReferencePoseType")
++EnumRedirects=(OldName="/Script/AnimNextAnimGraph.EAnimNextInjectionBlendMode", NewName="/Script/UAFAnimGraph.EAnimNextInjectionBlendMode")
++EnumRedirects=(OldName="/Script/AnimNextAnimGraph.EAnimNextInjectionType", NewName="/Script/UAFAnimGraph.EAnimNextInjectionType")
++EnumRedirects=(OldName="/Script/AnimNextAnimGraph.EAnimNextInjectionLifetimeType", NewName="/Script/UAFAnimGraph.EAnimNextInjectionLifetimeType")
++EnumRedirects=(OldName="/Script/AnimNextAnimGraph.EAnimNextInjectionStatus", NewName="/Script/UAFAnimGraph.EAnimNextInjectionStatus")
++EnumRedirects=(OldName="/Script/AnimNextAnimGraph.EAnimGroupSynchronizationRole", NewName="/Script/UAFAnimGraph.EAnimGroupSynchronizationRole")
++EnumRedirects=(OldName="/Script/AnimNextAnimGraph.EAnimGroupSynchronizationMode", NewName="/Script/UAFAnimGraph.EAnimGroupSynchronizationMode")
+
++StructRedirects=(OldName="/Script/AnimNextAnimGraphEditor.AnimNextSyncMarkerTraceInfo", NewName="/Script/UAFAnimGraphEditor.AnimNextSyncMarkerTraceInfo")
++StructRedirects=(OldName="/Script/AnimNextAnimGraphEditor.AnimNextSequenceTraceInfo", NewName="/Script/UAFAnimGraphEditor.AnimNextSequenceTraceInfo")
+
++ClassRedirects=(OldName="/Script/AnimNextAnimGraphEditor.AnimNextAnimationGraphFactory", NewName="/Script/UAFAnimGraphEditor.AnimNextAnimationGraphFactory")
++ClassRedirects=(OldName="/Script/AnimNextAnimGraphEditor.AssetDefinition_AnimNextAnimationGraphEntry", NewName="/Script/UAFAnimGraphEditor.AssetDefinition_AnimNextAnimationGraphEntry")
++ClassRedirects=(OldName="/Script/AnimNextAnimGraphEditor.AssetDefinition_AnimNextAnimationGraph", NewName="/Script/UAFAnimGraphEditor.AssetDefinition_AnimNextAnimationGraph")
++ClassRedirects=(OldName="/Script/AnimNextAnimGraphEditor.EvaluationProgramDetailsObject", NewName="/Script/UAFAnimGraphEditor.EvaluationProgramDetailsObject")
++ClassRedirects=(OldName="/Script/AnimNextAnimGraphEditor.SequenceInfoDetailsObject", NewName="/Script/UAFAnimGraphEditor.SequenceInfoDetailsObject")
++ClassRedirects=(OldName="/Script/AnimNextAnimGraphEditor.SequenceInfoDetailsObject", NewName="/Script/UAFAnimGraphEditor.SequenceInfoDetailsObject")
+
++StructRedirects=(OldName="/Script/AnimNextAnimGraphTestSuite.TestTraitSharedData", NewName="/Script/UAFAnimGraphTestSuite.TestTraitSharedData")
++StructRedirects=(OldName="/Script/AnimNextAnimGraphTestSuite.TestDerivedVector", NewName="/Script/UAFAnimGraphTestSuite.TestDerivedVector")
++StructRedirects=(OldName="/Script/AnimNextAnimGraphTestSuite.TraitWithOneChildSharedData", NewName="/Script/UAFAnimGraphTestSuite.TraitWithOneChildSharedData")
++StructRedirects=(OldName="/Script/AnimNextAnimGraphTestSuite.TraitWithChildrenSharedData", NewName="/Script/UAFAnimGraphTestSuite.TraitWithChildrenSharedData")
++StructRedirects=(OldName="/Script/AnimNextAnimGraphTestSuite.TraitA_BaseSharedData", NewName="/Script/UAFAnimGraphTestSuite.TraitA_BaseSharedData")
++StructRedirects=(OldName="/Script/AnimNextAnimGraphTestSuite.TraitAB_AddSharedData", NewName="/Script/UAFAnimGraphTestSuite.TraitAB_AddSharedData")
++StructRedirects=(OldName="/Script/AnimNextAnimGraphTestSuite.TraitAC_AddSharedData", NewName="/Script/UAFAnimGraphTestSuite.TraitAC_AddSharedData")
++StructRedirects=(OldName="/Script/AnimNextAnimGraphTestSuite.TraitSerialization_BaseSharedData", NewName="/Script/UAFAnimGraphTestSuite.TraitSerialization_BaseSharedData")
++StructRedirects=(OldName="/Script/AnimNextAnimGraphTestSuite.TraitSerialization_AddSharedData", NewName="/Script/UAFAnimGraphTestSuite.TraitSerialization_AddSharedData")
++StructRedirects=(OldName="/Script/AnimNextAnimGraphTestSuite.TraitNativeSerialization_AddSharedData", NewName="/Script/UAFAnimGraphTestSuite.TraitNativeSerialization_AddSharedData")
+
++StructRedirects=(OldName="/Script/AnimNextAnimGraphUncookedOnly.AnimNextAnimationGraphOutlinerData", NewName="/Script/UAFAnimGraphUncookedOnly.AnimNextAnimationGraphOutlinerData")
++StructRedirects=(OldName="/Script/AnimNextAnimGraphUncookedOnly.AnimNextSchemaAction_AddManifestNode", NewName="/Script/UAFAnimGraphUncookedOnly.AnimNextSchemaAction_AddManifestNode")
++StructRedirects=(OldName="/Script/AnimNextAnimGraphUncookedOnly.AnimNextSchemaAction_NotifyEvent", NewName="/Script/UAFAnimGraphUncookedOnly.AnimNextSchemaAction_NotifyEvent")
++StructRedirects=(OldName="/Script/AnimNextAnimGraphUncookedOnly.AnimNextBaseAction", NewName="/Script/UAFAnimGraphUncookedOnly.AnimNextBaseAction")
++StructRedirects=(OldName="/Script/AnimNextAnimGraphUncookedOnly.AnimNextManifestAction", NewName="/Script/UAFAnimGraphUncookedOnly.AnimNextManifestAction")
+
++ClassRedirects=(OldName="/Script/AnimNextAnimGraphUncookedOnly.AnimNextAnimGraphWorkspaceAssetUserData", NewName="/Script/UAFAnimGraphUncookedOnly.AnimNextAnimGraphWorkspaceAssetUserData")
++ClassRedirects=(OldName="/Script/AnimNextAnimGraphUncookedOnly.AnimNextController", NewName="/Script/UAFAnimGraphUncookedOnly.AnimNextController")
++ClassRedirects=(OldName="/Script/AnimNextAnimGraphUncookedOnly.AnimNextTraitStackUnitNode", NewName="/Script/UAFAnimGraphUncookedOnly.AnimNextTraitStackUnitNode")
++ClassRedirects=(OldName="/Script/AnimNextAnimGraphUncookedOnly.AnimNextAnimationGraphEntry", NewName="/Script/UAFAnimGraphUncookedOnly.AnimNextAnimationGraphEntry")
++ClassRedirects=(OldName="/Script/AnimNextAnimGraphUncookedOnly.AnimGraphNode_AnimNextGraph", NewName="/Script/UAFAnimGraphUncookedOnly.AnimGraphNode_AnimNextGraph")
++ClassRedirects=(OldName="/Script/AnimNextAnimGraphUncookedOnly.AnimNextAnimationGraphSchema", NewName="/Script/UAFAnimGraphUncookedOnly.AnimNextAnimationGraphSchema")
++ClassRedirects=(OldName="/Script/AnimNextAnimGraphUncookedOnly.AnimNextAnimationGraphLibrary", NewName="/Script/UAFAnimGraphUncookedOnly.AnimNextAnimationGraphLibrary")
++ClassRedirects=(OldName="/Script/AnimNextAnimGraphUncookedOnly.AnimNextAnimationGraph_EditorData", NewName="/Script/UAFAnimGraphUncookedOnly.AnimNextAnimationGraph_EditorData")
++ClassRedirects=(OldName="/Script/AnimNextAnimGraphUncookedOnly.K2Node_AnimNextInjection", NewName="/Script/UAFAnimGraphUncookedOnly.K2Node_AnimNextInjection")
++ClassRedirects=(OldName="/Script/AnimNextAnimGraphUncookedOnly.K2Node_AnimNextPlayAnim", NewName="/Script/UAFAnimGraphUncookedOnly.K2Node_AnimNextPlayAnim")
++ClassRedirects=(OldName="/Script/AnimNextAnimGraphUncookedOnly.K2Node_InjectionCallbackProxySetVariable", NewName="/Script/UAFAnimGraphUncookedOnly.K2Node_InjectionCallbackProxySetVariable")
\ No newline at end of file
diff --git a/Engine/Plugins/Experimental/UAF/UAFAnimGraph/Source/UAFAnimGraph/Internal/Graph/AnimNextGraphContextData.h b/Engine/Plugins/Experimental/UAF/UAFAnimGraph/Source/UAFAnimGraph/Internal/Graph/AnimNextGraphContextData.h
new file mode 100644
--- /dev/null
+++ b/Engine/Plugins/Experimental/UAF/UAFAnimGraph/Source/UAFAnimGraph/Internal/Graph/AnimNextGraphContextData.h
@@ -0,0 +1,32 @@
+// Copyright Epic Games, Inc. All Rights Reserved.
+
+#pragma once
+
+#include "CoreMinimal.h"
+#include "Module/AnimNextModuleContextData.h"
+#include "Graph/AnimNextGraphInstance.h"
+#include "AnimNextGraphContextData.generated.h"
+
+namespace UE::UAF
+{
+	struct FLatentPropertyHandle;
+}
+
+USTRUCT()
+struct FAnimNextGraphContextData : public FAnimNextModuleContextData
+{
+	GENERATED_BODY()
+
+	FAnimNextGraphContextData() = default;
+
+	FAnimNextGraphContextData(FAnimNextModuleInstance* InModuleInstance, const FAnimNextGraphInstance* InInstance)
+		: FAnimNextModuleContextData(InModuleInstance, InInstance)
+	{
+	}
+
+	const FAnimNextGraphInstance& GetGraphInstance() const { return static_cast<const FAnimNextGraphInstance&>(GetInstance()); }
+
+private:
+	friend struct FAnimNextGraphInstance;
+	friend struct FAnimNextExecuteContext;
+};
diff --git a/Engine/Plugins/Experimental/UAF/UAFAnimGraph/Source/UAFAnimGraph/Private/AnimNextAnimGraphModule.cpp b/Engine/Plugins/Experimental/UAF/UAFAnimGraph/Source/UAFAnimGraph/Private/AnimNextAnimGraphModule.cpp
new file mode 100644
--- /dev/null
+++ b/Engine/Plugins/Experimental/UAF/UAFAnimGraph/Source/UAFAnimGraph/Private/AnimNextAnimGraphModule.cpp
@@ -0,0 +1,345 @@
+// Copyright Epic Games, Inc. All Rights Reserved.
+
+#include "AnimNextAnimGraphModule.h"
+#include "AnimNextAnimGraphSettings.h"
+#include "Animation/BlendProfile.h"
+#include "Module/AnimNextModule.h"
+#include "Modules/ModuleInterface.h"
+#include "Modules/ModuleManager.h"
+#include "RigVMCore/RigVMRegistry.h"
+#include "Animation/AnimSequence.h"
+#include "Animation/BlendSpace.h"
+#include "Components/SkeletalMeshComponent.h"
+#include "Factory/AnimGraphFactory.h"
+#include "TraitInterfaces/IUpdate.h"
+#include "Graph/AnimNextAnimationGraph.h"
+#include "Graph/AnimNext_LODPose.h"
+#include "Graph/RigVMTrait_AnimNextPublicVariables.h"
+#include "Injection/InjectionRequest.h"
+#include "TraitCore/NodeTemplateRegistry.h"
+#include "TraitCore/TraitInterfaceRegistry.h"
+#include "TraitCore/TraitRegistry.h"
+#include "Traits/ModifyCurveTrait.h"
+#include "Traits/BlendSpacePlayerTraitData.h"
+#include "Traits/NotifyDispatcherTraitData.h"
+#include "Traits/SequencePlayer.h"
+#include "Factory/AnimNextFactoryParams.h"
+#include "Injection/InjectionSiteTrait.h"
+#include "Traits/BlendSmoother.h"
+#include "Traits/BlendStackTrait.h"
+
+#if WITH_ANIMNEXT_CONSOLE_COMMANDS
+#include "HAL/IConsoleManager.h"
+#include "UObject/UObjectIterator.h"
+
+#include "TraitCore/NodeDescription.h"
+#include "TraitCore/NodeTemplate.h"
+#include "TraitCore/TraitTemplate.h"
+#endif
+
+namespace UE::UAF::AnimGraph
+{
+
+void FAnimNextAnimGraphModule::StartupModule()
+{
+	// Ensure that AnimNext modules are loaded so we can correctly load plugin content
+	FModuleManager::LoadModuleChecked<IModuleInterface>("UAF");
+#if WITH_EDITORONLY_DATA
+	FModuleManager::LoadModuleChecked<IModuleInterface>("UAFUncookedOnly");
+#endif
+
+	// Setup default settings/factories
+	{
+		UAnimNextAnimGraphSettings* Settings = GetMutableDefault<UAnimNextAnimGraphSettings>();
+		Settings->LoadConfig();
+	}
+
+	FRigVMRegistry& RigVMRegistry = FRigVMRegistry::Get();
+	static TPair<UClass*, FRigVMRegistry::ERegisterObjectOperation> const AllowedObjectTypes[] =
+	{
+		{ UAnimNextAnimationGraph::StaticClass(), FRigVMRegistry::ERegisterObjectOperation::Class },
+		{ USkeletalMeshComponent::StaticClass(), FRigVMRegistry::ERegisterObjectOperation::ClassAndParents },
+	};
+
+	RigVMRegistry.RegisterObjectTypes(AllowedObjectTypes);
+	
+	static UScriptStruct* const AllowedStructTypes[] =
+	{
+		FAnimNextGraphReferencePose::StaticStruct(),
+		FRigVMTrait_AnimNextPublicVariables::StaticStruct(),
+		FAnimNextInjectionBlendSettings::StaticStruct(),
+		FModifyCurveParameters::StaticStruct()
+	};
+
+	RigVMRegistry.RegisterStructTypes(AllowedStructTypes);
+
+	FTraitRegistry::Init();
+	FTraitInterfaceRegistry::Init();
+	FNodeTemplateRegistry::Init();
+
+	// Register asset class -> graph factory mappings
+	{
+		// Register default graph host
+		FAnimNextFactoryParams Params;
+		Params.PushPublicTraitStruct<FBlendStackCoreTraitData>();
+		Params.PushPublicTraitStruct<FBlendSmootherCoreData>();
+		Params.PushPublicTraitStruct<FInjectionSiteTraitData>();
+
+		FAnimGraphFactory::RegisterDefaultParamsForClass(
+			UAnimNextAnimationGraph::StaticClass(),
+			MoveTemp(Params),
+			[](const FFactoryParamsInitializerContext& InContext)
+			{
+				ensure(InContext.AccessStruct<FInjectionSiteTraitData>([&InContext](FInjectionSiteTraitData& InStruct)
+				{
+					InStruct.Graph.Asset = CastChecked<UAnimNextAnimationGraph>(InContext.GetObject());
+				}));
+			});
+	}
+
+	{
+		// Register UAnimSequence 'player'
+		FAnimNextFactoryParams Params;
+		Params.PushPublicTraitStruct<FSequencePlayerData>();
+		Params.PushPublicTraitStruct<FNotifyDispatcherData>();
+		
+		FAnimGraphFactory::RegisterDefaultParamsForClass(
+			UAnimSequence::StaticClass(),
+			MoveTemp(Params),
+			[](const FFactoryParamsInitializerContext& InContext)
+			{
+				ensure(InContext.AccessStruct<FSequencePlayerData>([&InContext](FSequencePlayerData& InStruct)
+				{
+					InStruct.AnimSequence = CastChecked<UAnimSequence>(InContext.GetObject());
+				}));
+			});
+	}
+
+	{
+		// Register UBlendSpace 'player'
+		FAnimNextFactoryParams Params;
+		Params.PushPublicTraitStruct<FBlendSpacePlayerData>();
+		
+		FAnimGraphFactory::RegisterDefaultParamsForClass(
+			UBlendSpace::StaticClass(),
+			MoveTemp(Params),
+			[](const FFactoryParamsInitializerContext& InContext)
+			{
+				ensure(InContext.AccessStruct<FBlendSpacePlayerData>([&InContext](FBlendSpacePlayerData& InStruct)
+				{
+					InStruct.BlendSpace = CastChecked<UBlendSpace>(InContext.GetObject());
+				}));
+			});
+	}
+
+#if WITH_ANIMNEXT_CONSOLE_COMMANDS
+	if (!IsRunningCommandlet())
+	{
+		ConsoleCommands.Add(IConsoleManager::Get().RegisterConsoleCommand(
+			TEXT("UAF.ListNodeTemplates"),
+			TEXT("Dumps statistics about node templates to the log."),
+			FConsoleCommandWithArgsDelegate::CreateRaw(this, &FAnimNextAnimGraphModule::ListNodeTemplates),
+			ECVF_Default
+		));
+		ConsoleCommands.Add(IConsoleManager::Get().RegisterConsoleCommand(
+			TEXT("UAF.Module"),
+			TEXT("Dumps statistics about modules to the log."),
+			FConsoleCommandWithArgsDelegate::CreateRaw(this, &FAnimNextAnimGraphModule::ListAnimationGraphs),
+			ECVF_Default
+		));
+	}
+#endif
+}
+
+void FAnimNextAnimGraphModule::ShutdownModule()
+{
+#if WITH_ANIMNEXT_CONSOLE_COMMANDS
+	for (IConsoleObject* Cmd : ConsoleCommands)
+	{
+		IConsoleManager::Get().UnregisterConsoleObject(Cmd);
+	}
+	ConsoleCommands.Empty();
+#endif
+
+	FNodeTemplateRegistry::Destroy();
+	FTraitInterfaceRegistry::Destroy();
+	FTraitRegistry::Destroy();
+
+	LoadedGraphs.Reset();
+}
+
+void FAnimNextAnimGraphModule::AddReferencedObjects(FReferenceCollector& Collector)
+{
+	Collector.AddReferencedObjects(LoadedGraphs);
+}
+
+FString FAnimNextAnimGraphModule::GetReferencerName() const
+{
+	return TEXT("AnimNextAnimGraphModule");
+}
+
+#if WITH_ANIMNEXT_CONSOLE_COMMANDS
+void FAnimNextAnimGraphModule::ListNodeTemplates(const TArray<FString>& Args)
+{
+	// Turn off log times to make diff-ing easier
+	TGuardValue DisableLogTimes(GPrintLogTimes, ELogTimes::None);
+
+	// Make sure to log everything
+	const ELogVerbosity::Type OldVerbosity = LogAnimation.GetVerbosity();
+	LogAnimation.SetVerbosity(ELogVerbosity::All);
+
+	const FNodeTemplateRegistry& NodeTemplateRegistry = FNodeTemplateRegistry::Get();
+	const FTraitRegistry& TraitRegistry = FTraitRegistry::Get();
+
+	UE_LOG(LogAnimation, Log, TEXT("===== AnimNext Node Templates ====="));
+	UE_LOG(LogAnimation, Log, TEXT("Template Buffer Size: %u bytes"), NodeTemplateRegistry.TemplateBuffer.GetAllocatedSize());
+
+	for (auto It = NodeTemplateRegistry.TemplateUIDToHandleMap.CreateConstIterator(); It; ++It)
+	{
+		const FNodeTemplateRegistryHandle Handle = It.Value();
+		const FNodeTemplate* NodeTemplate = NodeTemplateRegistry.Find(Handle);
+
+		const uint32 NumTraits = NodeTemplate->GetNumTraits();
+
+		UE_LOG(LogAnimation, Log, TEXT("[%x] has %u traits ..."), NodeTemplate->GetUID(), NumTraits);
+		UE_LOG(LogAnimation, Log, TEXT("    Template Size: %u bytes"), NodeTemplate->GetNodeTemplateSize());
+		UE_LOG(LogAnimation, Log, TEXT("    Shared Data Size: %u bytes"), NodeTemplate->GetNodeSharedDataSize());
+		UE_LOG(LogAnimation, Log, TEXT("    Instance Data Size: %u bytes"), NodeTemplate->GetNodeInstanceDataSize());
+		UE_LOG(LogAnimation, Log, TEXT("    Traits ..."));
+
+		const FTraitTemplate* TraitTemplates = NodeTemplate->GetTraits();
+		for (uint32 TraitIndex = 0; TraitIndex < NumTraits; ++TraitIndex)
+		{
+			const FTraitTemplate* TraitTemplate = TraitTemplates + TraitIndex;
+			const FTrait* Trait = TraitRegistry.Find(TraitTemplate->GetRegistryHandle());
+			const FString TraitName = Trait != nullptr ? Trait->GetTraitName() : TEXT("<Unknown>");
+
+			const uint32 NextTraitIndex = TraitIndex + 1;
+			const uint32 EndOfNextTraitSharedData = NextTraitIndex < NumTraits ? TraitTemplates[NextTraitIndex].GetNodeSharedOffset() : NodeTemplate->GetNodeSharedDataSize();
+			const uint32 TraitSharedDataSize = EndOfNextTraitSharedData - TraitTemplate->GetNodeSharedOffset();
+
+			const uint32 EndOfNextTraitInstanceData = NextTraitIndex < NumTraits ? TraitTemplates[NextTraitIndex].GetNodeInstanceOffset() : NodeTemplate->GetNodeInstanceDataSize();
+			const uint32 TraitInstanceDataSize = EndOfNextTraitInstanceData - TraitTemplate->GetNodeInstanceOffset();
+
+			UE_LOG(LogAnimation, Log, TEXT("            %u: [%x] %s (%s)"), TraitIndex, TraitTemplate->GetUID().GetUID(), *TraitName, TraitTemplate->GetMode() == ETraitMode::Base ? TEXT("Base") : TEXT("Additive"));
+			UE_LOG(LogAnimation, Log, TEXT("                Shared Data: [Offset: %u bytes, Size: %u bytes]"), TraitTemplate->GetNodeSharedOffset(), TraitSharedDataSize);
+			if (TraitTemplate->HasLatentProperties() && Trait != nullptr)
+			{
+				UE_LOG(LogAnimation, Log, TEXT("                Shared Data Latent Property Handles: [Offset: %u bytes, Count: %u]"), TraitTemplate->GetNodeSharedLatentPropertyHandlesOffset(), Trait->GetNumLatentTraitProperties());
+			}
+			UE_LOG(LogAnimation, Log, TEXT("                Instance Data: [Offset: %u bytes, Size: %u bytes]"), TraitTemplate->GetNodeInstanceOffset(), TraitInstanceDataSize);
+		}
+	}
+
+	LogAnimation.SetVerbosity(OldVerbosity);
+}
+
+void FAnimNextAnimGraphModule::ListAnimationGraphs(const TArray<FString>& Args)
+{
+	// Turn off log times to make diff-ing easier
+	TGuardValue DisableLogTimes(GPrintLogTimes, ELogTimes::None);
+
+	// Make sure to log everything
+	const ELogVerbosity::Type OldVerbosity = LogAnimation.GetVerbosity();
+	LogAnimation.SetVerbosity(ELogVerbosity::All);
+
+	TArray<const UAnimNextAnimationGraph*> AnimationGraphs;
+
+	for (TObjectIterator<UAnimNextAnimationGraph> It; It; ++It)
+	{
+		AnimationGraphs.Add(*It);
+	}
+
+	struct FCompareObjectNames
+	{
+		FORCEINLINE bool operator()(const UAnimNextAnimationGraph& Lhs, const UAnimNextAnimationGraph& Rhs) const
+		{
+			return Lhs.GetPathName().Compare(Rhs.GetPathName()) < 0;
+		}
+	};
+	AnimationGraphs.Sort(FCompareObjectNames());
+
+	const FNodeTemplateRegistry& NodeTemplateRegistry = FNodeTemplateRegistry::Get();
+	const FTraitRegistry& TraitRegistry = FTraitRegistry::Get();
+	const bool bDetailedOutput = true;
+
+	UE_LOG(LogAnimation, Log, TEXT("===== AnimNext Modules ====="));
+	UE_LOG(LogAnimation, Log, TEXT("Num Graphs: %u"), AnimationGraphs.Num());
+
+	for (const UAnimNextAnimationGraph* AnimationGraph : AnimationGraphs)
+	{
+		uint32 TotalInstanceSize = 0;
+		uint32 NumNodes = 0;
+		{
+			// We always have a node at offset 0
+			int32 NodeOffset = 0;
+
+			while (NodeOffset < AnimationGraph->SharedDataBuffer.Num())
+			{
+				const FNodeDescription* NodeDesc = reinterpret_cast<const FNodeDescription*>(&AnimationGraph->SharedDataBuffer[NodeOffset]);
+
+				TotalInstanceSize += NodeDesc->GetNodeInstanceDataSize();
+				NumNodes++;
+
+				const FNodeTemplate* NodeTemplate = NodeTemplateRegistry.Find(NodeDesc->GetTemplateHandle());
+				NodeOffset += NodeTemplate->GetNodeSharedDataSize();
+			}
+		}
+
+		UE_LOG(LogAnimation, Log, TEXT("    %s ..."), *AnimationGraph->GetPathName());
+		UE_LOG(LogAnimation, Log, TEXT("        Shared Data Size: %.2f KB"), double(AnimationGraph->SharedDataBuffer.Num()) / 1024.0);
+		UE_LOG(LogAnimation, Log, TEXT("        Max Instance Data Size: %.2f KB"), double(TotalInstanceSize) / 1024.0);
+		UE_LOG(LogAnimation, Log, TEXT("        Num Nodes: %u"), NumNodes);
+
+		if (bDetailedOutput)
+		{
+			// We always have a node at offset 0
+			int32 NodeOffset = 0;
+
+			while (NodeOffset < AnimationGraph->SharedDataBuffer.Num())
+			{
+				const FNodeDescription* NodeDesc = reinterpret_cast<const FNodeDescription*>(&AnimationGraph->SharedDataBuffer[NodeOffset]);
+				const FNodeTemplate* NodeTemplate = NodeTemplateRegistry.Find(NodeDesc->GetTemplateHandle());
+
+				const uint32 NumTraits = NodeTemplate->GetNumTraits();
+
+				UE_LOG(LogAnimation, Log, TEXT("        Node %u: [Template %x with %u traits]"), NodeDesc->GetUID().GetNodeIndex(), NodeTemplate->GetUID(), NumTraits);
+				UE_LOG(LogAnimation, Log, TEXT("            Shared Data: [Offset: %u bytes, Size: %u bytes]"), NodeOffset, NodeTemplate->GetNodeSharedDataSize());
+				UE_LOG(LogAnimation, Log, TEXT("            Instance Data Size: %u bytes"), NodeDesc->GetNodeInstanceDataSize());
+				UE_LOG(LogAnimation, Log, TEXT("            Traits ..."));
+
+				const FTraitTemplate* TraitTemplates = NodeTemplate->GetTraits();
+				for (uint32 TraitIndex = 0; TraitIndex < NumTraits; ++TraitIndex)
+				{
+					const FTraitTemplate* TraitTemplate = TraitTemplates + TraitIndex;
+					const FTrait* Trait = TraitRegistry.Find(TraitTemplate->GetRegistryHandle());
+					const FString TraitName = Trait != nullptr ? Trait->GetTraitName() : TEXT("<Unknown>");
+
+					const uint32 NextTraitIndex = TraitIndex + 1;
+					const uint32 EndOfNextTraitSharedData = NextTraitIndex < NumTraits ? TraitTemplates[NextTraitIndex].GetNodeSharedOffset() : NodeTemplate->GetNodeSharedDataSize();
+					const uint32 TraitSharedDataSize = EndOfNextTraitSharedData - TraitTemplate->GetNodeSharedOffset();
+
+					const uint32 EndOfNextTraitInstanceData = NextTraitIndex < NumTraits ? TraitTemplates[NextTraitIndex].GetNodeInstanceOffset() : NodeTemplate->GetNodeInstanceDataSize();
+					const uint32 TraitInstanceDataSize = EndOfNextTraitInstanceData - TraitTemplate->GetNodeInstanceOffset();
+
+					UE_LOG(LogAnimation, Log, TEXT("                    %u: [%x] %s (%s)"), TraitIndex, TraitTemplate->GetUID().GetUID(), *TraitName, TraitTemplate->GetMode() == ETraitMode::Base ? TEXT("Base") : TEXT("Additive"));
+					UE_LOG(LogAnimation, Log, TEXT("                        Shared Data: [Offset: %u bytes, Size: %u bytes]"), TraitTemplate->GetNodeSharedOffset(), TraitSharedDataSize);
+					if (TraitTemplate->HasLatentProperties() && Trait != nullptr)
+					{
+						UE_LOG(LogAnimation, Log, TEXT("                        Shared Data Latent Property Handles: [Offset: %u bytes, Count: %u]"), TraitTemplate->GetNodeSharedLatentPropertyHandlesOffset(), Trait->GetNumLatentTraitProperties());
+					}
+					UE_LOG(LogAnimation, Log, TEXT("                        Instance Data: [Offset: %u bytes, Size: %u bytes]"), TraitTemplate->GetNodeInstanceOffset(), TraitInstanceDataSize);
+				}
+
+				NodeOffset += NodeTemplate->GetNodeSharedDataSize();
+			}
+		}
+	}
+
+	LogAnimation.SetVerbosity(OldVerbosity);
+}
+#endif
+
+}
+
+IMPLEMENT_MODULE(UE::UAF::AnimGraph::FAnimNextAnimGraphModule, UAFAnimGraph)
diff --git a/Engine/Plugins/Experimental/UAF/UAFAnimGraph/Source/UAFAnimGraph/Private/Factory/AnimGraphFactory.cpp b/Engine/Plugins/Experimental/UAF/UAFAnimGraph/Source/UAFAnimGraph/Private/Factory/AnimGraphFactory.cpp
new file mode 100644
--- /dev/null
+++ b/Engine/Plugins/Experimental/UAF/UAFAnimGraph/Source/UAFAnimGraph/Private/Factory/AnimGraphFactory.cpp
@@ -0,0 +1,145 @@
+﻿// Copyright Epic Games, Inc. All Rights Reserved.
+
+#include "Factory/AnimGraphFactory.h"
+
+#include "Factory/AnimGraphBuilderContext.h"
+#include "Containers/StripedMap.h"
+#include "Factory/AnimNextAnimGraphBuilder.h"
+#include "Graph/AnimNextAnimationGraph.h"
+#include "TraitCore/TraitWriter.h"
+
+namespace UE::UAF::Private
+{
+
+// References to all generated factory graphs, keyed by method
+static TStripedMap<32, uint64, TWeakObjectPtr<const UAnimNextAnimationGraph>> GAnimationGraphMap;
+
+// All registered object -> param mappings
+static TMap<const UClass*, TPair<FAnimNextFactoryParams, FAnimGraphFactory::FParamsInitializer>> GDefaultParamsForObject;
+static TArray<UClass*> GAllFactoryClasses;
+
+const TPair<FAnimNextFactoryParams, FAnimGraphFactory::FParamsInitializer>* FindParamsForClass(const UClass* InClass)
+{
+	const UClass* CurrentClass = InClass;
+	const TPair<FAnimNextFactoryParams, FAnimGraphFactory::FParamsInitializer>* FoundParams = nullptr;
+	while (FoundParams == nullptr && CurrentClass != nullptr)
+	{
+		FoundParams = GDefaultParamsForObject.Find(CurrentClass);
+		CurrentClass = CurrentClass->GetSuperClass();
+	}
+
+	return FoundParams;
+}
+
+}
+
+namespace UE::UAF
+{
+
+uint64 IAnimGraphBuilder::GetKey() const
+{
+	return 0;
+}
+
+const UAnimNextAnimationGraph* FAnimGraphFactory::GetDefaultGraphHost()
+{
+	const TPair<FAnimNextFactoryParams, FParamsInitializer>* ParamsTask = Private::FindParamsForClass(UAnimNextAnimationGraph::StaticClass());
+	return BuildGraph(ParamsTask->Key.GetBuilder());
+}
+
+const UAnimNextAnimationGraph* FAnimGraphFactory::GetOrBuildGraph(const UObject* InObject, const IAnimGraphBuilder& InBuilder)
+{
+	if (const UAnimNextAnimationGraph* AnimationGraph = Cast<UAnimNextAnimationGraph>(InObject))
+	{
+		return AnimationGraph;
+	}
+
+	return BuildGraph(InBuilder);
+}
+
+const UAnimNextAnimationGraph* FAnimGraphFactory::BuildGraph(const IAnimGraphBuilder& InBuilder)
+{
+	uint64 Key = InBuilder.GetKey();
+	if (Key == 0)
+	{
+		return nullptr;
+	}
+
+	auto ProduceGraph = [&InBuilder]() -> TWeakObjectPtr<const UAnimNextAnimationGraph>
+	{
+		FAnimGraphBuilderContext Context;
+		if (InBuilder.Build(Context))
+		{
+			return Context.Build();
+		}
+		return nullptr;
+	};
+
+	TWeakObjectPtr<const UAnimNextAnimationGraph> AnimationGraph;
+	auto ApplyWeakGraph = [&ProduceGraph, &AnimationGraph](TWeakObjectPtr<const UAnimNextAnimationGraph>& InFoundAnimationGraph)
+	{
+		if (InFoundAnimationGraph.IsStale())
+		{
+			// Found a stale weak ptr, so regenerate (this may be a procedural graph or a previously loaded graph that has been GCed)
+			InFoundAnimationGraph = ProduceGraph();
+		}
+		AnimationGraph = InFoundAnimationGraph;
+	};
+
+	Private::GAnimationGraphMap.FindOrProduceAndApplyForWrite(Key, ProduceGraph, ApplyWeakGraph);
+	return AnimationGraph.Get();
+}
+
+const FAnimNextFactoryParams& FAnimGraphFactory::GetDefaultParamsForClass(const UClass* InClass)
+{
+	check(InClass);
+	const TPair<FAnimNextFactoryParams, FParamsInitializer>* ParamsTask = Private::FindParamsForClass(InClass);
+	check(ParamsTask != nullptr);	// Unregistered object type
+	return ParamsTask->Key;
+}
+
+FAnimNextFactoryParams FAnimGraphFactory::GetDefaultParamsForObject(const UObject* InObject)
+{
+	check(InObject);
+	const TPair<FAnimNextFactoryParams, FParamsInitializer>* ParamsTask = Private::FindParamsForClass(InObject->GetClass());
+	check(ParamsTask != nullptr);	// Unregistered object type
+
+	FAnimNextFactoryParams ParamsCopy = ParamsTask->Key;
+	if (ParamsCopy.Builder.Stacks.Num() > 0)
+	{
+		FFactoryParamsInitializerContext Context;
+		Context.Object = InObject;
+		Context.StructDatas = ParamsCopy.Builder.Stacks[0].TraitStructs;
+		ParamsTask->Value(Context);
+	}
+	return ParamsCopy;
+}
+
+void FAnimGraphFactory::InitializeDefaultParamsForObject(const UObject* InObject, FAnimNextFactoryParams& InOutParams)
+{
+	check(InObject);
+	const TPair<FAnimNextFactoryParams, FParamsInitializer>* ParamsTask = Private::FindParamsForClass(InObject->GetClass());
+	check(ParamsTask != nullptr);	// Unregistered object type
+
+	if (InOutParams.Builder.Stacks.Num() > 0)
+	{
+		FFactoryParamsInitializerContext Context;
+		Context.Object = InObject;
+		Context.StructDatas = InOutParams.Builder.Stacks[0].TraitStructs;
+		ParamsTask->Value(Context);
+	}
+}
+
+void FAnimGraphFactory::RegisterDefaultParamsForClass(const UClass* InClass, FAnimNextFactoryParams&& InParams, FParamsInitializer&& InInitializer)
+{
+	Private::GDefaultParamsForObject.Add(InClass, { MoveTemp(InParams), MoveTemp(InInitializer) });
+	Private::GAllFactoryClasses.Add(const_cast<UClass*>(InClass));
+}
+
+TConstArrayView<UClass*> FAnimGraphFactory::GetRegisteredClasses()
+{
+	return Private::GAllFactoryClasses;
+}
+
+
+}
diff --git a/Engine/Plugins/Experimental/UAF/UAFAnimGraph/Source/UAFAnimGraph/Private/Factory/AnimNextFactoryParams.cpp b/Engine/Plugins/Experimental/UAF/UAFAnimGraph/Source/UAFAnimGraph/Private/Factory/AnimNextFactoryParams.cpp
new file mode 100644
--- /dev/null
+++ b/Engine/Plugins/Experimental/UAF/UAFAnimGraph/Source/UAFAnimGraph/Private/Factory/AnimNextFactoryParams.cpp
@@ -0,0 +1,31 @@
+﻿// Copyright Epic Games, Inc. All Rights Reserved.
+
+#include "Factory/AnimNextFactoryParams.h"
+#include "Graph/AnimNextGraphInstance.h"
+
+void FAnimNextFactoryParams::InitializeInstance(FUAFAssetInstance& InInstance) const
+{
+	if (Builder.Stacks.Num() > 0)
+	{
+		for (const TInstancedStruct<FAnimNextTraitSharedData>& InitializePayload : Builder.Stacks[0].TraitStructs)
+		{
+			InInstance.AccessVariablesStruct(InitializePayload.GetScriptStruct(), [&InitializePayload](FStructView InStructView)
+			{
+				check(InitializePayload.GetScriptStruct() == InStructView.GetScriptStruct());
+				InStructView.GetScriptStruct()->CopyScriptStruct(InStructView.GetMemory(), InitializePayload.GetMemory());
+			});
+		}
+	}
+
+	const UE::UAF::FInstanceTaskContext TaskContext(InInstance);
+	for (const UE::UAF::FInstanceTask& InitializeTask : InitializeTasks)
+	{
+		InitializeTask(TaskContext);
+	}
+}
+
+void FAnimNextFactoryParams::Reset()
+{
+	Builder.Reset();
+	InitializeTasks.Empty();
+}
\ No newline at end of file
diff --git a/Engine/Plugins/Experimental/UAF/UAFAnimGraph/Source/UAFAnimGraph/Private/Factory/AnimNextSimpleAnimGraphBuilder.cpp b/Engine/Plugins/Experimental/UAF/UAFAnimGraph/Source/UAFAnimGraph/Private/Factory/AnimNextSimpleAnimGraphBuilder.cpp
new file mode 100644
--- /dev/null
+++ b/Engine/Plugins/Experimental/UAF/UAFAnimGraph/Source/UAFAnimGraph/Private/Factory/AnimNextSimpleAnimGraphBuilder.cpp
@@ -0,0 +1,221 @@
+﻿// Copyright Epic Games, Inc. All Rights Reserved.
+
+#include "Factory/AnimNextSimpleAnimGraphBuilder.h"
+
+#include "Factory/AnimGraphBuilderContext.h"
+#include "Misc/HashBuilder.h"
+#include "TraitCore/NodeHandle.h"
+#include "TraitCore/NodeTemplate.h"
+#include "TraitCore/TraitRegistry.h"
+#include "TraitCore/NodeTemplateBuilder.h"
+#include "TraitCore/TraitWriter.h"
+#include "Traits/ReferencePoseTrait.h"
+
+bool FAnimNextSimpleAnimGraphBuilder::Build(UE::UAF::FAnimGraphBuilderContext& InContext) const
+{
+	using namespace UE::UAF;
+
+	ensure(Stacks.Num() == 1);	// Currently only support a single stack
+
+	if (Stacks.Num() != 1 || Stacks[0].TraitStructs.Num() == 0)
+	{
+		return false;
+	}
+
+	const FTraitRegistry& TraitRegistry = FTraitRegistry::Get();
+
+	InContext.VariableStructs.Reserve(Stacks[0].TraitStructs.Num());
+
+	struct FTraitStackData
+	{
+		FNodeHandle NodeHandle;
+		TArray<TInstancedStruct<FAnimNextTraitSharedData>, TInlineAllocator<4>> TraitStructs;
+		TArray<FTraitUID, TInlineAllocator<4>> TraitUIDs;
+		TMap<TPair<uint32, FName>, int32, TInlineSetAllocator<64>> VariableIndexMap;
+	};
+
+	TArray<FTraitStackData, TInlineAllocator<4>> TraitStacks;
+	TraitStacks.Reserve(Stacks.Num());
+
+	{
+		FTraitStackData TraitStackData;
+		TraitStackData.TraitStructs = Stacks[0].TraitStructs;
+		TraitStackData.TraitUIDs.Reserve(Stacks[0].TraitStructs.Num());
+		TraitStacks.Add(MoveTemp(TraitStackData));
+	}
+
+	for (int32 TraitStackIndex = 0; TraitStackIndex < TraitStacks.Num(); ++TraitStackIndex)
+	{
+		FTraitStackData& TraitStackData = TraitStacks[TraitStackIndex];
+		int32 LatentPropertyIndex = 0;
+		int32 TraitIndex = 0;
+		for (TInstancedStruct<FAnimNextTraitSharedData>& TraitStruct : TraitStackData.TraitStructs)
+		{
+			const FTrait* Trait = TraitRegistry.Find(TraitStruct.GetScriptStruct());
+			TraitStackData.TraitUIDs.Add(Trait->GetTraitUID());
+			if (TraitStackIndex == 0)
+			{
+				// Add the variables of the single trait stack we support as 'public interface' structs
+				InContext.VariableStructs.Add(TraitStruct.GetScriptStruct());
+			}
+
+			for (TFieldIterator<FProperty> It(TraitStruct.GetScriptStruct()); It; ++It)
+			{
+				TraitStackData.VariableIndexMap.Add({TraitIndex, It->GetFName()}, LatentPropertyIndex);
+				LatentPropertyIndex++;
+			}
+			TraitIndex++;
+		}
+	}
+
+	// Register nodes
+	auto RegisterNode = [&InContext](FTraitStackData& InTraitStackData)
+	{
+		TArray<uint8> NodeTemplateBuffer;
+		const FNodeTemplate* NodeTemplate = FNodeTemplateBuilder::BuildNodeTemplate(InTraitStackData.TraitUIDs, NodeTemplateBuffer);
+		InTraitStackData.NodeHandle = InContext.TraitWriter.RegisterNode(*NodeTemplate);
+	};
+
+	const int32 NumTraitStacks = TraitStacks.Num();
+	for (int32 TraitStackIndex = 0; TraitStackIndex < NumTraitStacks; ++TraitStackIndex)
+	{
+		FTraitStackData& TraitStackData = TraitStacks[TraitStackIndex];
+		RegisterNode(TraitStackData);
+
+		// Single trait we support is always the root
+		if (TraitStackIndex == 0)
+		{
+			InContext.RootTraitHandle = FAnimNextEntryPointHandle(TraitStackData.NodeHandle);
+		}
+
+		// Add new trait stack with refpose trait to cover 'unlinked' child pins
+		const FTrait* Trait = TraitRegistry.Find(FAnimNextReferencePoseTraitSharedData::StaticStruct());
+		const FTraitUID RefPoseUID = Trait->GetTraitUID();
+		for (TInstancedStruct<FAnimNextTraitSharedData>& TraitStruct : TraitStackData.TraitStructs)
+		{
+			for (TFieldIterator<FProperty> It(TraitStruct.GetScriptStruct()); It; ++It)
+			{
+				FStructProperty* StructProperty = CastField<FStructProperty>(*It);
+				if (StructProperty && StructProperty->Struct == FAnimNextTraitHandle::StaticStruct())
+				{
+					FTraitStackData RefPoseTraitStackData;
+					TInstancedStruct<FAnimNextTraitSharedData> InstancedStruct;
+					InstancedStruct.InitializeAsScriptStruct(FAnimNextReferencePoseTraitSharedData::StaticStruct());
+
+					RefPoseTraitStackData.TraitStructs.Add(MoveTemp(InstancedStruct));
+					RefPoseTraitStackData.TraitUIDs.Add(RefPoseUID);
+
+					RegisterNode(RefPoseTraitStackData);
+
+					// Link to this new stack
+					FAnimNextTraitHandle* Handle = It->ContainerPtrToValuePtr<FAnimNextTraitHandle>(TraitStruct.GetMutableMemory());
+					*Handle = FAnimNextTraitHandle(RefPoseTraitStackData.NodeHandle);
+
+					TraitStacks.Add(RefPoseTraitStackData);
+				}
+			}
+		}
+	}
+
+	// Write traits
+	InContext.TraitWriter.BeginNodeWriting();
+	for (int32 TraitStackIndex = 0; TraitStackIndex < TraitStacks.Num(); ++TraitStackIndex)
+	{
+		const FTraitStackData& TraitStackData = TraitStacks[TraitStackIndex];
+		InContext.TraitWriter.WriteNode(TraitStackData.NodeHandle,
+			[&TraitStackData](uint32 InTraitIndex, FName InPropertyName)
+			{
+				return TraitStackData.VariableIndexMap.FindRef({InTraitIndex, InPropertyName});
+			},
+			[&TraitStackData](uint32 InTraitIndex)
+			{
+				return TConstStructView<FAnimNextTraitSharedData>(TraitStackData.TraitStructs[InTraitIndex]);
+			});
+	}
+	InContext.TraitWriter.EndNodeWriting();
+
+	ensure(InContext.TraitWriter.GetErrorState() == FTraitWriter::EErrorState::None);
+
+	return true;
+}
+
+void FAnimNextSimpleAnimGraphBuilder::ValidateTraitStruct(int32 InStackIndex, TConstStructView<FAnimNextTraitSharedData> InStruct)
+{
+	using namespace UE::UAF;
+
+	ensure(InStackIndex == 0);	// Currently only support a single stack
+	if (InStackIndex != 0)
+	{
+		return;
+	}
+
+	if (Stacks.Num() == 0)
+	{
+		Stacks.SetNum(InStackIndex + 1);
+	}
+
+	const FTraitRegistry& TraitRegistry = FTraitRegistry::Get();
+	const FTrait* Trait = TraitRegistry.Find(InStruct.GetScriptStruct());
+	if (Trait == nullptr)
+	{
+		ensureMsgf(false, TEXT("FAnimNextSimpleAnimGraphBuilder: Could not find registered trait"));
+		return;
+	}
+
+	if (Stacks[InStackIndex].TraitStructs.Num() == 0 && Trait->GetTraitMode() != ETraitMode::Base)
+	{
+		ensureMsgf(false, TEXT("FAnimNextSimpleAnimGraphBuilder: First trait should always be a base trait"));
+		return;
+	}
+}
+
+void FAnimNextSimpleAnimGraphBuilder::PushTraitStructViewToStack(int32 InStackIndex, TConstStructView<FAnimNextTraitSharedData> InStruct)
+{
+	using namespace UE::UAF;
+
+	ValidateTraitStruct(InStackIndex, InStruct);
+
+	auto FindExistingTrait = [&InStruct](const TInstancedStruct<FAnimNextTraitSharedData>& InExistingStruct)
+	{
+		return InStruct.GetScriptStruct() == InExistingStruct.GetScriptStruct();
+	};
+
+	if (TInstancedStruct<FAnimNextTraitSharedData>* ExistingStruct = Stacks[InStackIndex].TraitStructs.FindByPredicate(FindExistingTrait))
+	{
+		*ExistingStruct = InStruct;
+	}
+	else
+	{
+		Stacks[InStackIndex].TraitStructs.Emplace(InStruct);
+	}
+
+	InvalidateKey();
+}
+
+void FAnimNextSimpleAnimGraphBuilder::PushTraitInstancedStructToStack(int32 InStackIndex, TInstancedStruct<FAnimNextTraitSharedData>&& InStruct)
+{
+	using namespace UE::UAF;
+
+	ValidateTraitStruct(InStackIndex, InStruct);
+
+	auto FindExistingTrait = [&InStruct](const TInstancedStruct<FAnimNextTraitSharedData>& InExistingStruct)
+	{
+		return InStruct.GetScriptStruct() == InExistingStruct.GetScriptStruct();
+	};
+
+	if (TInstancedStruct<FAnimNextTraitSharedData>* ExistingStruct = Stacks[InStackIndex].TraitStructs.FindByPredicate(FindExistingTrait))
+	{
+		*ExistingStruct = InStruct;
+	}
+	else
+	{
+		Stacks[InStackIndex].TraitStructs.Emplace(MoveTemp(InStruct));
+	}
+
+	InvalidateKey();
+}
+
+uint64 FAnimNextSimpleAnimGraphBuilder::RecalculateKey() const
+{
+	return UE::StructUtils::GetStructHash64(TConstStructView<FAnimNextSimpleAnimGraphBuilder>(*this));
+}
diff --git a/Engine/Plugins/Experimental/UAF/UAFAnimGraph/Source/UAFAnimGraph/Private/Graph/AnimNextGraphInstance.cpp b/Engine/Plugins/Experimental/UAF/UAFAnimGraph/Source/UAFAnimGraph/Private/Graph/AnimNextGraphInstance.cpp
new file mode 100644
--- /dev/null
+++ b/Engine/Plugins/Experimental/UAF/UAFAnimGraph/Source/UAFAnimGraph/Private/Graph/AnimNextGraphInstance.cpp
@@ -0,0 +1,237 @@
+// Copyright Epic Games, Inc. All Rights Reserved.
+
+#include "Graph/AnimNextGraphInstance.h"
+
+#include "AnimNextAnimGraphStats.h"
+#include "ObjectTrace.h"
+#include "UAFRigVMComponent.h"
+#include "Graph/AnimNextAnimationGraph.h"
+#include "Graph/AnimNextGraphContextData.h"
+#include "Graph/AnimNextGraphLatentPropertiesContextData.h"
+#include "TraitCore/ExecutionContext.h"
+#include "Graph/GC_GraphInstanceComponent.h"
+#include "Graph/RigUnit_AnimNextShimRoot.h"
+#include "Graph/RigVMTrait_AnimNextPublicVariables.h"
+#include "Logging/StructuredLog.h"
+#include "Module/AnimNextModule.h"
+#include "Module/AnimNextModuleInstance.h"
+
+DEFINE_STAT(STAT_AnimNext_Graph_RigVM);
+
+FAnimNextGraphInstance::FAnimNextGraphInstance()
+{
+}
+
+FAnimNextGraphInstance::~FAnimNextGraphInstance()
+{
+	TRACE_INSTANCE_LIFETIME_END(ModuleInstance ? ModuleInstance->GetObject() : nullptr, GetUniqueId());
+	Release();
+}
+
+void FAnimNextGraphInstance::Release()
+{
+#if WITH_EDITORONLY_DATA
+	if(const UAnimNextAnimationGraph* Graph = GetAnimationGraph())
+	{
+		FScopeLock Lock(&Graph->GraphInstancesLock);
+		Graph->GraphInstances.Remove(this);
+	}
+#endif
+
+	if (!GraphInstancePtr.IsValid())
+	{
+		return;
+	}
+
+	GraphInstancePtr.Reset();
+	ModuleInstance = nullptr;
+	RootGraphInstance = nullptr;
+	ReleaseComponents();
+	Asset = nullptr;
+}
+
+bool FAnimNextGraphInstance::IsValid() const
+{
+	return GraphInstancePtr.IsValid();
+}
+
+const UAnimNextAnimationGraph* FAnimNextGraphInstance::GetAnimationGraph() const
+{
+	return CastChecked<UAnimNextAnimationGraph>(Asset);
+}
+
+FName FAnimNextGraphInstance::GetEntryPoint() const
+{
+	return EntryPoint;
+}
+
+UE::UAF::FWeakTraitPtr FAnimNextGraphInstance::GetGraphRootPtr() const
+{
+	return GraphInstancePtr;
+}
+
+FAnimNextModuleInstance* FAnimNextGraphInstance::GetModuleInstance() const
+{
+	return ModuleInstance;
+}
+
+FAnimNextGraphInstance* FAnimNextGraphInstance::GetParentGraphInstance() const
+{
+	if(IsRoot())
+	{
+		return nullptr;
+	}
+	else
+	{
+		return static_cast<FAnimNextGraphInstance*>(HostInstance);
+	}
+}
+
+FAnimNextGraphInstance* FAnimNextGraphInstance::GetRootGraphInstance() const
+{
+	return RootGraphInstance;
+}
+
+bool FAnimNextGraphInstance::UsesAnimationGraph(const UAnimNextAnimationGraph* InAnimationGraph) const
+{
+	return GetAnimationGraph() == InAnimationGraph;
+}
+
+bool FAnimNextGraphInstance::UsesEntryPoint(FName InEntryPoint) const
+{
+	if(const UAnimNextAnimationGraph* AnimationGraph = GetAnimationGraph())
+	{
+		if(InEntryPoint == NAME_None)
+		{
+			return EntryPoint == AnimationGraph->DefaultEntryPoint;
+		}
+
+		return InEntryPoint == EntryPoint;
+	}
+	return false;
+}
+
+bool FAnimNextGraphInstance::IsRoot() const
+{	
+	return this == RootGraphInstance;
+}
+
+bool FAnimNextGraphInstance::HasUpdated() const
+{
+	return bHasUpdatedOnce;
+}
+
+void FAnimNextGraphInstance::MarkAsUpdated()
+{
+	bHasUpdatedOnce = true;
+}
+
+void FAnimNextGraphInstance::ExecuteLatentPins(const TConstArrayView<UE::UAF::FLatentPropertyHandle>& LatentHandles, void* DestinationBasePtr, bool bIsFrozen, bool bJustBecameRelevant)
+{
+	SCOPE_CYCLE_COUNTER(STAT_AnimNext_Graph_RigVM);
+
+	if (!IsValid())
+	{
+		return;
+	}
+
+	if (URigVM* VM = GetAnimationGraph()->RigVM)
+	{
+		FUAFRigVMComponent& RigVMComponent = GetComponent<FUAFRigVMComponent>();
+		FRigVMExtendedExecuteContext& ExtendedExecuteContext = RigVMComponent.GetExtendedExecuteContext();
+		FAnimNextExecuteContext& AnimNextContext = ExtendedExecuteContext.GetPublicDataSafe<FAnimNextExecuteContext>();
+		if (ModuleInstance)
+		{
+			AnimNextContext.SetOwningObject(ModuleInstance->GetObject());
+		}
+
+		// Insert our context data for the scope of execution
+		FAnimNextGraphLatentPropertiesContextData ContextData(ModuleInstance, this, LatentHandles, DestinationBasePtr, bIsFrozen, bJustBecameRelevant);
+		UE::UAF::FScopedExecuteContextData ContextDataScope(AnimNextContext, ContextData);
+
+		VM->ExecuteVM(ExtendedExecuteContext, FRigUnit_AnimNextShimRoot::EventName);
+	}
+}
+
+void FAnimNextGraphInstance::CopyVariablesToLatentPins(const TConstArrayView<UE::UAF::FLatentPropertyHandle>& LatentHandles, void* DestinationBasePtr, bool bIsFrozen, bool bJustBecameRelevant)
+{
+	for (UE::UAF::FLatentPropertyHandle Handle : LatentHandles)
+	{
+		if (!Handle.IsIndexValid())
+		{
+			// This handle isn't valid
+			continue;
+		}
+
+		if (!Handle.IsOffsetValid())
+		{
+			// This handle isn't valid
+			continue;
+		}
+
+		if (bIsFrozen && Handle.CanFreeze())
+		{
+			// This handle can freeze and we are frozen, no need to update it
+			continue;
+		}
+
+		if (!bJustBecameRelevant && Handle.OnBecomeRelevant())
+		{
+			// This handle should only update on become relevant
+			continue;
+		}
+
+		uint8* DestinationPtr = (uint8*)DestinationBasePtr + Handle.GetLatentPropertyOffset();
+		AccessVariablePropertyByIndex(Handle.GetLatentPropertyIndex(), [DestinationPtr](const FProperty* InProperty, TArrayView<uint8> InData)
+		{
+			// Copy from our source into our destination
+			// We assume the source and destination properties are identical
+			check(InData.Num() == InProperty->GetElementSize());
+			InProperty->CopyCompleteValue(DestinationPtr, InData.GetData());
+		});
+	}
+}
+
+#if WITH_EDITORONLY_DATA
+void FAnimNextGraphInstance::Freeze()
+{
+	if (!IsValid())
+	{
+		return;
+	}
+
+	GraphInstancePtr.Reset();
+	ReleaseComponents();
+	bHasUpdatedOnce = false;
+}
+
+void FAnimNextGraphInstance::Thaw()
+{
+	if (const UAnimNextAnimationGraph* AnimationGraph = GetAnimationGraph())
+	{
+		if (Variables.bHasBeenInitialized)
+		{
+			MigrateVariables();
+		}
+		else
+		{
+			InitializeVariables();
+		}
+
+		{
+			UE::UAF::FExecutionContext Context(*this);
+			if(const FAnimNextTraitHandle* FoundHandle = AnimationGraph->ResolvedRootTraitHandles.Find(EntryPoint))
+			{
+				GraphInstancePtr = Context.AllocateNodeInstance(*this, *FoundHandle);
+			}
+		}
+
+		if (!IsValid())
+		{
+			// We failed to allocate our instance, clear everything
+			Release();
+		}
+	}
+}
+
+#endif
diff --git a/Engine/Plugins/Experimental/UAF/UAFAnimGraph/Source/UAFAnimGraph/Private/Graph/GC_GraphInstanceComponent.cpp b/Engine/Plugins/Experimental/UAF/UAFAnimGraph/Source/UAFAnimGraph/Private/Graph/GC_GraphInstanceComponent.cpp
new file mode 100644
--- /dev/null
+++ b/Engine/Plugins/Experimental/UAF/UAFAnimGraph/Source/UAFAnimGraph/Private/Graph/GC_GraphInstanceComponent.cpp
@@ -0,0 +1,49 @@
+// Copyright Epic Games, Inc. All Rights Reserved.
+
+#include "Graph/GC_GraphInstanceComponent.h"
+
+#include "TraitCore/ExecutionContext.h"
+#include "TraitInterfaces/IGarbageCollection.h"
+
+void FGCGraphInstanceComponent::Register(const UE::UAF::FWeakTraitPtr& InTraitPtr)
+{
+	TraitsWithReferences.Add(InTraitPtr);
+}
+
+void FGCGraphInstanceComponent::Unregister(const UE::UAF::FWeakTraitPtr& InTraitPtr)
+{
+	using namespace UE::UAF;
+	
+	const int32 EntryIndex = TraitsWithReferences.IndexOfByPredicate(
+		[&InTraitPtr](const FWeakTraitPtr& TraitPtr)
+		{
+			return TraitPtr == InTraitPtr;
+		});
+
+	if (ensure(EntryIndex != INDEX_NONE))
+	{
+		TraitsWithReferences.RemoveAtSwap(EntryIndex);
+	}
+}
+
+void FGCGraphInstanceComponent::AddStructReferencedObjects(FReferenceCollector& Collector)
+{
+	using namespace UE::UAF;
+	
+	FExecutionContext Context;
+	FTraitStackBinding TraitStack;
+	TTraitBinding<IGarbageCollection> GCTrait;
+
+	// TODO: If we kept the entries sorted by graph instance, we could re-use the execution context
+	for (const FWeakTraitPtr& TraitPtr : TraitsWithReferences)
+	{
+		Context.BindTo(TraitPtr);
+
+		if (Context.GetStack(TraitPtr, TraitStack))
+		{
+			ensure(TraitStack.GetInterface(GCTrait));
+
+			GCTrait.AddReferencedObjects(Context, Collector);
+		}
+	}
+}
diff --git a/Engine/Plugins/Experimental/UAF/UAFAnimGraph/Source/UAFAnimGraph/Private/Graph/GC_GraphInstanceComponent.h b/Engine/Plugins/Experimental/UAF/UAFAnimGraph/Source/UAFAnimGraph/Private/Graph/GC_GraphInstanceComponent.h
new file mode 100644
--- /dev/null
+++ b/Engine/Plugins/Experimental/UAF/UAFAnimGraph/Source/UAFAnimGraph/Private/Graph/GC_GraphInstanceComponent.h
@@ -0,0 +1,46 @@
+// Copyright Epic Games, Inc. All Rights Reserved.
+
+#pragma once
+
+#include "CoreMinimal.h"
+#include "TraitCore/TraitPtr.h"
+#include "Graph/UAFGraphInstanceComponent.h"
+#include "GC_GraphInstanceComponent.generated.h"
+
+class FReferenceCollector;
+
+/**
+ * FGCGraphInstanceComponent
+ *
+ * This component maintains the necessary state to garbage collection.
+ */
+USTRUCT()
+struct FGCGraphInstanceComponent : public FUAFGraphInstanceComponent
+{
+	GENERATED_BODY()
+
+	FGCGraphInstanceComponent() = default;
+
+	// Registers the provided trait with the GC system
+	// Once registered, IGarbageCollection::AddReferencedObjects will be called on it during GC
+	void Register(const UE::UAF::FWeakTraitPtr& InTraitPtr);
+
+	// Unregisters the provided trait from the GC system
+	void Unregister(const UE::UAF::FWeakTraitPtr& InTraitPtr);
+
+	// Called during garbage collection to collect strong object references
+	void AddStructReferencedObjects(FReferenceCollector& Collector);
+
+private:
+	// List of trait handles that contain UObject references and implement IGarbageCollection
+	TArray<UE::UAF::FWeakTraitPtr> TraitsWithReferences;
+};
+
+template<>
+struct TStructOpsTypeTraits<FGCGraphInstanceComponent> : public TStructOpsTypeTraitsBase2<FGCGraphInstanceComponent>
+{
+	enum
+	{
+		WithAddStructReferencedObjects = true,
+	};
+};
diff --git a/Engine/Plugins/Experimental/UAF/UAFAnimGraph/Source/UAFAnimGraph/Private/Graph/TraceAnimNextGraphInstances.h b/Engine/Plugins/Experimental/UAF/UAFAnimGraph/Source/UAFAnimGraph/Private/Graph/TraceAnimNextGraphInstances.h
new file mode 100644
--- /dev/null
+++ b/Engine/Plugins/Experimental/UAF/UAFAnimGraph/Source/UAFAnimGraph/Private/Graph/TraceAnimNextGraphInstances.h
@@ -0,0 +1,27 @@
+// Copyright Epic Games, Inc. All Rights Reserved.
+#pragma once
+
+#include "EvaluationVM/EvaluationProgram.h"
+#include "RewindDebugger/AnimNextTrace.h"
+
+#if ANIMNEXT_TRACE_ENABLED
+
+struct FAnimNextGraphInstance;
+
+namespace UE::UAF
+{
+	struct FEvaluationProgram;
+	
+	void TraceGraphInstances(const FAnimNextGraphInstance& RootGraph);
+	void TraceEvaluationProgram(const FEvaluationProgram& Program, const FAnimNextGraphInstance& RootGraph);
+}
+
+#define TRACE_ANIMNEXT_GRAPHINSTANCES(RootGraph) UE::UAF::TraceGraphInstances(RootGraph);
+#define TRACE_ANIMNEXT_EVALUATIONPROGRAM(Program, RootGraph) UE::UAF::TraceEvaluationProgram(Program, RootGraph);
+
+#else
+
+#define TRACE_ANIMNEXT_GRAPHINSTANCES(RootGraph)
+#define TRACE_ANIMNEXT_EVALUATIONPROGRAM(Program, RootGraph)
+
+#endif
diff --git a/Engine/Plugins/Experimental/UAF/UAFAnimGraph/Source/UAFAnimGraph/Private/Injection/GraphInstanceInjectionComponent.cpp b/Engine/Plugins/Experimental/UAF/UAFAnimGraph/Source/UAFAnimGraph/Private/Injection/GraphInstanceInjectionComponent.cpp
new file mode 100644
--- /dev/null
+++ b/Engine/Plugins/Experimental/UAF/UAFAnimGraph/Source/UAFAnimGraph/Private/Injection/GraphInstanceInjectionComponent.cpp
@@ -0,0 +1,11 @@
+﻿// Copyright Epic Games, Inc. All Rights Reserved.
+
+#include "GraphInstanceInjectionComponent.h"
+#include "Graph/AnimNextGraphInstance.h"
+
+FGraphInstanceInjectionComponent::FGraphInstanceInjectionComponent()
+	: FUAFGraphInstanceComponent()
+	, InjectionInfo(GetGraphInstance())
+{
+}
+
diff --git a/Engine/Plugins/Experimental/UAF/UAFAnimGraph/Source/UAFAnimGraph/Private/Injection/InjectionUtils.cpp b/Engine/Plugins/Experimental/UAF/UAFAnimGraph/Source/UAFAnimGraph/Private/Injection/InjectionUtils.cpp
new file mode 100644
--- /dev/null
+++ b/Engine/Plugins/Experimental/UAF/UAFAnimGraph/Source/UAFAnimGraph/Private/Injection/InjectionUtils.cpp
@@ -0,0 +1,203 @@
+﻿// Copyright Epic Games, Inc. All Rights Reserved.
+
+#include "Injection/InjectionUtils.h"
+#include "GraphInterfaces/AnimNextNativeDataInterface_AnimSequencePlayer.h"
+#include "Animation/AnimSequence.h"
+#include "Component/AnimNextComponent.h"
+#include "Logging/StructuredLog.h"
+
+namespace UE::UAF
+{
+	FInjectionRequestPtr FInjectionUtils::Inject(UAnimNextComponent* InComponent, FInjectionRequestArgs&& InArgs, FInjectionLifetimeEvents&& InLifetimeEvents)
+	{
+		return Inject(InComponent, InComponent->GetModuleHandle(), MoveTemp(InArgs), MoveTemp(InLifetimeEvents));
+	}
+
+	FInjectionRequestPtr FInjectionUtils::Inject(UObject* InHost, FModuleHandle InHandle, FInjectionRequestArgs&& InArgs, FInjectionLifetimeEvents&& InLifetimeEvents)
+	{
+		check(IsInGameThread());
+
+		FInjectionRequestPtr Request = MakeInjectionRequest();
+		if(Request->Inject(MoveTemp(InArgs), MoveTemp(InLifetimeEvents), InHost, InHandle))
+		{
+			return Request;
+		}
+		return nullptr;
+	}
+
+	void FInjectionUtils::Uninject(FInjectionRequestPtr InInjectionRequest)
+	{
+		if(!InInjectionRequest.IsValid())
+		{
+			return;
+		}
+		InInjectionRequest->Uninject();
+	}
+
+	FInjectionRequestPtr FInjectionUtils::PlayAnim(
+		UAnimNextComponent* InComponent,
+		const FInjectionSite& InSite,
+		UAnimSequence* InAnimSequence,
+		float InPlayRate,
+		float InStartPosition,
+		const FInjectionBlendSettings& InBlendInSettings,
+		const FInjectionBlendSettings& InBlendOutSettings,
+		FInjectionLifetimeEvents&& InLifetimeEvents)
+	{
+		return PlayAnimHandle(
+			InComponent,
+			InComponent->GetModuleHandle(),
+			InSite,
+			InAnimSequence,
+			{
+				.PlayRate = InPlayRate,
+				.StartPosition = InStartPosition
+			},
+			InBlendInSettings,
+			InBlendOutSettings,
+			MoveTemp(InLifetimeEvents));
+	}
+
+	FInjectionRequestPtr FInjectionUtils::PlayAnim(
+		UAnimNextComponent* InComponent,
+		const FInjectionSite& InSite,
+		UAnimSequence* InAnimSequence,
+		FPlayAnimArgs&& InArgs,
+		const FInjectionBlendSettings& InBlendInSettings,
+		const FInjectionBlendSettings& InBlendOutSettings,
+		FInjectionLifetimeEvents&& InLifetimeEvents)
+	{
+		return PlayAnimHandle(
+			InComponent,
+			InComponent->GetModuleHandle(),
+			InSite,
+			InAnimSequence,
+			MoveTemp(InArgs),
+			InBlendInSettings,
+			InBlendOutSettings,
+			MoveTemp(InLifetimeEvents));
+	}
+
+	FInjectionRequestPtr FInjectionUtils::PlayAnimHandle(
+		UObject* InHost,
+		FModuleHandle InModuleHandle,
+		const FInjectionSite& InSite,
+		UAnimSequence* InAnimSequence,
+		FPlayAnimArgs&& InArgs,
+		const FInjectionBlendSettings& InBlendInSettings,
+		const FInjectionBlendSettings& InBlendOutSettings,
+		FInjectionLifetimeEvents&& InLifetimeEvents)
+	{
+		check(IsInGameThread());
+
+		FInjectionRequestArgs RequestArgs;
+		RequestArgs.Site = InSite;
+		RequestArgs.Object = InAnimSequence;
+		RequestArgs.Type = EAnimNextInjectionType::InjectObject;
+		RequestArgs.BlendInSettings = InBlendInSettings;
+		RequestArgs.BlendOutSettings = InBlendOutSettings;
+		RequestArgs.LifetimeType = InArgs.LifetimeType;
+
+		if (!InArgs.FactoryParams.IsValid())
+		{
+			RequestArgs.FactoryParams =
+				FAnimNextFactoryParams()
+				.PushPublicTrait(FSequencePlayerData(
+					{
+						.AnimSequence = InAnimSequence,
+						.PlayRate = InArgs.PlayRate,
+						.StartPosition = InArgs.StartPosition,
+						.bLoop = [&InArgs, &InAnimSequence]()
+						{
+							switch (InArgs.LoopMode)
+							{
+							case ELoopMode::Auto:
+								return InAnimSequence->bLoop;
+							case ELoopMode::ForceLoop:
+								return true;
+							case ELoopMode::ForceNonLoop:
+								return false;
+							}
+							return false;
+						}()
+					}));
+		}
+		else
+		{
+			RequestArgs.FactoryParams = MoveTemp(InArgs.FactoryParams);
+		}
+
+		return Inject(InHost, InModuleHandle, MoveTemp(RequestArgs), MoveTemp(InLifetimeEvents));
+	}
+
+	FInjectionRequestPtr FInjectionUtils::InjectAsset(
+		UAnimNextComponent* InComponent,
+		const FInjectionSite& InSite,
+		UObject* InAsset,
+		FAnimNextFactoryParams&& InFactoryParams,
+		const FInjectionBlendSettings& InBlendInSettings,
+		const FInjectionBlendSettings& InBlendOutSettings,
+		FInjectionLifetimeEvents&& InLifetimeEvents)
+	{
+		return InjectAsset(
+			InComponent,
+			InComponent->GetModuleHandle(),
+			InSite,
+			InAsset,
+			MoveTemp(InFactoryParams),
+			InBlendInSettings,
+			InBlendOutSettings,
+			MoveTemp(InLifetimeEvents));
+	}
+
+	FInjectionRequestPtr FInjectionUtils::InjectAsset(
+		UObject* InHost,
+		FModuleHandle InModuleHandle,
+		const FInjectionSite& InSite,
+		UObject* InAsset,
+		FAnimNextFactoryParams&& InFactoryParams,
+		const FInjectionBlendSettings& InBlendInSettings,
+		const FInjectionBlendSettings& InBlendOutSettings,
+		FInjectionLifetimeEvents&& InLifetimeEvents)
+	{
+		check(IsInGameThread());
+
+		FInjectionRequestArgs RequestArgs;
+		RequestArgs.Site = InSite;
+		RequestArgs.Object = InAsset;
+		RequestArgs.Type = EAnimNextInjectionType::InjectObject;
+		RequestArgs.BlendInSettings = InBlendInSettings;
+		RequestArgs.BlendOutSettings = InBlendOutSettings;
+		RequestArgs.FactoryParams = MoveTemp(InFactoryParams);
+
+		return Inject(InHost, InModuleHandle, MoveTemp(RequestArgs), MoveTemp(InLifetimeEvents));
+	}
+
+	FInjectionRequestPtr FInjectionUtils::InjectEvaluationModifier(
+		UAnimNextComponent* InComponent,
+		const TSharedRef<IEvaluationModifier>& InEvaluationModifier,
+		const FInjectionSite& InSite)
+	{
+		return InjectEvaluationModifier(
+			InComponent,
+			InComponent->GetModuleHandle(),
+			InEvaluationModifier,
+			InSite);
+	}
+
+	FInjectionRequestPtr FInjectionUtils::InjectEvaluationModifier(
+		UObject* InHost,
+		FModuleHandle InModuleHandle,
+		const TSharedRef<IEvaluationModifier>& InEvaluationModifier,
+		const FInjectionSite& InSite)
+	{
+		check(IsInGameThread());
+
+		FInjectionRequestArgs RequestArgs;
+		RequestArgs.Site = InSite;
+		RequestArgs.Type = EAnimNextInjectionType::EvaluationModifier;
+		RequestArgs.EvaluationModifier = InEvaluationModifier;
+
+		return Inject(InHost, InModuleHandle, MoveTemp(RequestArgs));
+	}
+}
diff --git a/Engine/Plugins/Experimental/UAF/UAFAnimGraph/Source/UAFAnimGraph/Private/TraitInterfaces/IGarbageColllection.cpp b/Engine/Plugins/Experimental/UAF/UAFAnimGraph/Source/UAFAnimGraph/Private/TraitInterfaces/IGarbageColllection.cpp
new file mode 100644
--- /dev/null
+++ b/Engine/Plugins/Experimental/UAF/UAFAnimGraph/Source/UAFAnimGraph/Private/TraitInterfaces/IGarbageColllection.cpp
@@ -0,0 +1,45 @@
+// Copyright Epic Games, Inc. All Rights Reserved.
+
+#include "TraitInterfaces/IGarbageCollection.h"
+
+#include "TraitCore/ExecutionContext.h"
+#include "Graph/GC_GraphInstanceComponent.h"
+
+namespace UE::UAF
+{
+	AUTO_REGISTER_ANIM_TRAIT_INTERFACE(IGarbageCollection)
+
+#if WITH_EDITOR
+	const FText& IGarbageCollection::GetDisplayName() const
+	{
+		static FText InterfaceName = NSLOCTEXT("TraitInterfaces", "TraitInterface_IGarbageCollection_Name", "Garbage Collection");
+		return InterfaceName;
+	}
+	const FText& IGarbageCollection::GetDisplayShortName() const
+	{
+		static FText InterfaceShortName = NSLOCTEXT("TraitInterfaces", "TraitInterface_IGarbageCollection_ShortName", "GC");
+		return InterfaceShortName;
+	}
+#endif // WITH_EDITOR
+
+	void IGarbageCollection::RegisterWithGC(const FExecutionContext& Context, const FTraitBinding& Binding)
+	{
+		FGCGraphInstanceComponent& Component = Context.GetComponent<FGCGraphInstanceComponent>();
+		Component.Register(Binding.GetTraitPtr());
+	}
+
+	void IGarbageCollection::UnregisterWithGC(const FExecutionContext& Context, const FTraitBinding& Binding)
+	{
+		FGCGraphInstanceComponent& Component = Context.GetComponent<FGCGraphInstanceComponent>();
+		Component.Unregister(Binding.GetTraitPtr());
+	}
+
+	void IGarbageCollection::AddReferencedObjects(const FExecutionContext& Context, const TTraitBinding<IGarbageCollection>& Binding, FReferenceCollector& Collector) const
+	{
+		TTraitBinding<IGarbageCollection> SuperBinding;
+		if (Binding.GetStackInterfaceSuper(SuperBinding))
+		{
+			SuperBinding.AddReferencedObjects(Context, Collector);
+		}
+	}
+}
diff --git a/Engine/Plugins/Experimental/UAF/UAFAnimGraph/Source/UAFAnimGraph/Private/TraitInterfaces/IGraphFactory.cpp b/Engine/Plugins/Experimental/UAF/UAFAnimGraph/Source/UAFAnimGraph/Private/TraitInterfaces/IGraphFactory.cpp
new file mode 100644
--- /dev/null
+++ b/Engine/Plugins/Experimental/UAF/UAFAnimGraph/Source/UAFAnimGraph/Private/TraitInterfaces/IGraphFactory.cpp
@@ -0,0 +1,65 @@
+// Copyright Epic Games, Inc. All Rights Reserved.
+
+#include "TraitInterfaces/IGraphFactory.h"
+
+#include "Factory/AnimGraphFactory.h"
+#include "Factory/AnimNextFactoryParams.h"
+#include "Graph/AnimNextAnimationGraph.h"
+#include "TraitCore/ExecutionContext.h"
+
+namespace UE::UAF
+{
+	AUTO_REGISTER_ANIM_TRAIT_INTERFACE(IGraphFactory)
+
+#if WITH_EDITOR
+	const FText& IGraphFactory::GetDisplayName() const
+	{
+		static FText InterfaceName = NSLOCTEXT("TraitInterfaces", "TraitInterface_IGraphFactory_Name", "Graph Factory");
+		return InterfaceName;
+	}
+	const FText& IGraphFactory::GetDisplayShortName() const
+	{
+		static FText InterfaceShortName = NSLOCTEXT("TraitInterfaces", "TraitInterface_IGraphFactory_ShortName", "GF");
+		return InterfaceShortName;
+	}
+#endif // WITH_EDITOR
+
+	void IGraphFactory::GetFactoryParams(FExecutionContext& Context, const TTraitBinding<IGraphFactory>& Binding, const UObject* InObject, FAnimNextFactoryParams& InOutParams) const
+	{
+		TTraitBinding<IGraphFactory> SuperBinding;
+		if (Binding.GetStackInterfaceSuper(SuperBinding))
+		{
+			SuperBinding.GetFactoryParams(Context, InObject, InOutParams);
+		}
+	}
+
+	const UAnimNextAnimationGraph* IGraphFactory::GetOrBuildGraph(FExecutionContext& Context, const FTraitBinding& Binding, const UObject* InObject, FAnimNextFactoryParams& InOutParams)
+	{
+		// If a graph was provided directly, just use it
+		if (const UAnimNextAnimationGraph* AnimationGraph = Cast<UAnimNextAnimationGraph>(InObject))
+		{
+			return AnimationGraph;
+		}
+		else if (InObject != nullptr) // Otherwise attempt to build a graph procedurally from the supplied params
+		{
+			if (!InOutParams.IsValid())
+			{
+				InOutParams = FAnimGraphFactory::GetDefaultParamsForClass(InObject->GetClass());
+			}
+
+			// Apply default initializer (trait stacks can still alter this below)
+			FAnimGraphFactory::InitializeDefaultParamsForObject(InObject, InOutParams);
+
+			// Defer to the trait stack to gather any additional traits to instantiate
+			TTraitBinding<IGraphFactory> GraphBuilderBinding;
+			if (Binding.GetStackInterface(GraphBuilderBinding))
+			{
+				GraphBuilderBinding.GetFactoryParams(Context, InObject, InOutParams);
+			}
+
+			return FAnimGraphFactory::BuildGraph(InOutParams.GetBuilder());
+		}
+
+		return nullptr;
+	}
+}
diff --git a/Engine/Plugins/Experimental/UAF/UAFAnimGraph/Source/UAFAnimGraph/Public/EvaluationVM/Tasks/BlendKeyframesPerBone.h b/Engine/Plugins/Experimental/UAF/UAFAnimGraph/Source/UAFAnimGraph/Public/EvaluationVM/Tasks/BlendKeyframesPerBone.h
new file mode 100644
--- /dev/null
+++ b/Engine/Plugins/Experimental/UAF/UAFAnimGraph/Source/UAFAnimGraph/Public/EvaluationVM/Tasks/BlendKeyframesPerBone.h
@@ -0,0 +1,130 @@
+// Copyright Epic Games, Inc. All Rights Reserved.
+
+#pragma once
+
+#include "Animation/AnimCurveTypes.h"
+#include "Animation/NamedValueArray.h"
+#include "CoreMinimal.h"
+#include "EvaluationVM/EvaluationTask.h"
+#include "EvaluationVM/Tasks/BlendKeyframes.h"
+#include "Animation/AttributesContainer.h"
+
+#include "BlendKeyframesPerBone.generated.h"
+
+#define UE_API UAFANIMGRAPH_API
+
+struct FBlendSampleData;
+class IBlendProfileInterface;
+class IInterpolationIndexProvider;
+class USkeleton;
+
+/*
+ * Blend Overwrite Keyframe Per Bone With Scale Task
+ *
+ * This pops the top keyframe from the VM keyframe stack, it scales each bone by a factor, and pushes
+ * back the result onto the stack.
+ * Top = Top * ScaleFactor
+ * 
+ * If no blend profile is provided, this task behaves like FAnimNextBlendOverwriteKeyframeWithScaleTask
+ * Note that rotations will not be normalized after this task.
+ */
+USTRUCT()
+struct FAnimNextBlendOverwriteKeyframePerBoneWithScaleTask : public FAnimNextBlendOverwriteKeyframeWithScaleTask
+{
+	GENERATED_BODY()
+
+	DECLARE_ANIM_EVALUATION_TASK(FAnimNextBlendOverwriteKeyframePerBoneWithScaleTask)
+
+	static UE_API FAnimNextBlendOverwriteKeyframePerBoneWithScaleTask Make(const IBlendProfileInterface* BlendProfile, const FBlendSampleData& BlendData, float ScaleFactor);
+
+	// Task entry point
+	UE_API virtual void Execute(UE::UAF::FEvaluationVM& VM) const override;
+
+	// The blend data associated with the keyframe to overwrite
+	// Cannot use a UPROPERTY with a pointer to a struct
+	//UPROPERTY()
+	const FBlendSampleData* BlendData = nullptr;
+
+	// The blend profile to use
+	// An optional blend profile provides a blend weight per bone
+	const IBlendProfileInterface* BlendProfile = nullptr;
+
+	const USkeleton* SourceSkeleton = nullptr;
+};
+
+/*
+ * Blend Add Keyframe Per Bone With Scale Task
+ *
+ * This pops the top two keyframes (A and B) from the VM keyframe stack (let B be at the top).
+ * B is our intermediary result that we add on top of; while A is the keyframe we scale.
+ * The result is pushed back onto the stack.
+ * Top = Top + (Top-1 * ScaleFactor)
+ * 
+ * If no blend profile is provided, this task behaves like FAnimNextBlendAddKeyframeWithScaleTask
+ * Note that rotations will not be normalized after this task.
+ */
+USTRUCT()
+struct FAnimNextBlendAddKeyframePerBoneWithScaleTask : public FAnimNextBlendAddKeyframeWithScaleTask
+{
+	GENERATED_BODY()
+
+	DECLARE_ANIM_EVALUATION_TASK(FAnimNextBlendAddKeyframePerBoneWithScaleTask)
+
+	static UE_API FAnimNextBlendAddKeyframePerBoneWithScaleTask Make(const IBlendProfileInterface* BlendProfile, const FBlendSampleData& BlendDataA, const FBlendSampleData& BlendDataB, float ScaleFactor);
+
+	// Task entry point
+	UE_API virtual void Execute(UE::UAF::FEvaluationVM& VM) const override;
+
+	// The blend data associated with the keyframe A
+	// Cannot use a UPROPERTY with a pointer to a struct
+	//UPROPERTY()
+	const FBlendSampleData* BlendDataA = nullptr;
+
+	// The blend data associated with the keyframe B
+	// Cannot use a UPROPERTY with a pointer to a struct
+	//UPROPERTY()
+	const FBlendSampleData* BlendDataB = nullptr;
+
+	// The blend profile to use
+	// An optional blend profile provides a blend weight per bone
+	const IBlendProfileInterface* BlendProfile = nullptr;
+};
+
+
+USTRUCT()
+struct FAnimNextBlendKeyframePerBoneWithScaleTask : public FAnimNextBlendAddKeyframeWithScaleTask
+{
+	GENERATED_BODY()
+
+	DECLARE_ANIM_EVALUATION_TASK(FAnimNextBlendKeyframePerBoneWithScaleTask)
+
+	struct FMaskedAttributeWeight
+	{
+		UE::Anim::FAttributeId Attribute;
+		float Weight;
+
+		FMaskedAttributeWeight(const UE::Anim::FAttributeId InAttribute, const float InWeight)
+			: Attribute(InAttribute)
+			, Weight(InWeight)
+		{
+		}
+	};
+
+	static UE_API FAnimNextBlendKeyframePerBoneWithScaleTask Make(TSharedPtr<IInterpolationIndexProvider> BlendProfile, const USkeleton* Skeleton, const TArray<float>& BoneMaskWeights, const UE::Anim::TNamedValueArray<FDefaultAllocator, UE::Anim::FCurveElement>& CurveMaskWeights, const TArray<FMaskedAttributeWeight>& AttributeMaskWeights, float ScaleFactor);
+
+	// Task entry point
+	UE_API virtual void Execute(UE::UAF::FEvaluationVM& VM) const override;
+
+private:
+	TArray<float> BoneMaskWeights;
+
+	UE::Anim::TNamedValueArray<FDefaultAllocator, UE::Anim::FCurveElement> CurveMaskWeights;
+
+	TArray<FMaskedAttributeWeight> AttributeMaskWeights;
+
+	TSharedPtr<IInterpolationIndexProvider> BlendProfile = nullptr;
+
+	const USkeleton* SourceSkeleton = nullptr;
+};
+
+#undef UE_API
diff --git a/Engine/Plugins/Experimental/UAF/UAFAnimGraph/Source/UAFAnimGraph/Public/Factory/AnimGraphFactory.h b/Engine/Plugins/Experimental/UAF/UAFAnimGraph/Source/UAFAnimGraph/Public/Factory/AnimGraphFactory.h
new file mode 100644
--- /dev/null
+++ b/Engine/Plugins/Experimental/UAF/UAFAnimGraph/Source/UAFAnimGraph/Public/Factory/AnimGraphFactory.h
@@ -0,0 +1,86 @@
+﻿// Copyright Epic Games, Inc. All Rights Reserved.
+
+#pragma once
+
+#include "CoreTypes.h"
+#include "Factory/AnimNextFactoryParams.h"
+
+class UAnimNextAnimationGraph;
+class UObject;
+
+namespace UE::UAF
+{
+	class FTraitWriter;
+	struct IAnimGraphBuilder;
+	struct FAnimGraphFactory;
+}
+
+namespace UE::UAF
+{
+
+struct FFactoryParamsInitializerContext
+{
+	// Get the object used to initialize the params (e.g. AnimSequence for a player)
+	const UObject* GetObject() const
+	{
+		return Object;
+	}
+
+	// Access a struct during initialization (for instance to inject an AnimSequence)
+	template<typename StructType>
+	bool AccessStruct(TFunctionRef<void(StructType&)> InFunction) const
+	{
+		for (TInstancedStruct<FAnimNextTraitSharedData>& StructData : StructDatas)
+		{
+			if (StructData.GetScriptStruct() == TBaseStructure<StructType>::Get())
+			{
+				InFunction(StructData.GetMutable<StructType>());
+				return true;
+			}
+		}
+		return false;
+	}
+
+private:
+	friend FAnimGraphFactory;
+
+	TArrayView<TInstancedStruct<FAnimNextTraitSharedData>> StructDatas;
+	const UObject* Object = nullptr;
+};
+
+// Creates or recycles programmatically-generated graphs
+// Uses the hash of the recipe to determine if the graph has been created already.
+// If the hash matches one that has already been created but graph has already been GCed it will be re-created as needed.
+struct FAnimGraphFactory
+{
+	// Make a graph from the specified recipe
+	// If the object is an UAnimNextAnimationGraph, it will be returned rather than build a new one
+	UAFANIMGRAPH_API static const UAnimNextAnimationGraph* GetOrBuildGraph(const UObject* InObject, const IAnimGraphBuilder& InBuilder);
+
+	// Make a graph from the specified recipe
+	UAFANIMGRAPH_API static const UAnimNextAnimationGraph* BuildGraph(const IAnimGraphBuilder& InBuilder);
+
+	// Make the default graph for running a UAnimNextAnimationGraph asset (allowing graphs to blend between each other)
+	UAFANIMGRAPH_API static const UAnimNextAnimationGraph* GetDefaultGraphHost();
+
+	// Get default params used to build & initialize a graph according to registered object class
+	UAFANIMGRAPH_API static const FAnimNextFactoryParams& GetDefaultParamsForClass(const UClass* InClass);
+
+	// Get default params used to build & initialize a graph according to registered object class, copying InObject via the registered initializer
+	// Effective equivalent of GetDefaultParamsForClass and InitializeDefaultParamsForObject.
+	UAFANIMGRAPH_API static FAnimNextFactoryParams GetDefaultParamsForObject(const UObject* InObject);
+
+	// Applies default initializer to the params (usually injecting the object pointer) 
+	UAFANIMGRAPH_API static void InitializeDefaultParamsForObject(const UObject* InObject, FAnimNextFactoryParams& InOutParams);
+
+	// Function used to initialize the object for the factory-generated graph
+	using FParamsInitializer = TFunction<void(const FFactoryParamsInitializerContext&)>;
+
+	// Register default params used to build & initialize a graph according to the supplied object class
+	UAFANIMGRAPH_API static void RegisterDefaultParamsForClass(const UClass* InClass, FAnimNextFactoryParams&& InParams, FParamsInitializer&& InInitializer);
+
+	// Get all registered classes used to generate factory graphs
+	UAFANIMGRAPH_API static TConstArrayView<UClass*> GetRegisteredClasses();
+};
+
+}
\ No newline at end of file
diff --git a/Engine/Plugins/Experimental/UAF/UAFAnimGraph/Source/UAFAnimGraph/Public/TraitInterfaces/IGraphFactory.h b/Engine/Plugins/Experimental/UAF/UAFAnimGraph/Source/UAFAnimGraph/Public/TraitInterfaces/IGraphFactory.h
new file mode 100644
--- /dev/null
+++ b/Engine/Plugins/Experimental/UAF/UAFAnimGraph/Source/UAFAnimGraph/Public/TraitInterfaces/IGraphFactory.h
@@ -0,0 +1,60 @@
+// Copyright Epic Games, Inc. All Rights Reserved.
+
+#pragma once
+
+#include "CoreMinimal.h"
+#include "StructUtils/StructView.h"
+#include "TraitCore/ITraitInterface.h"
+#include "TraitCore/TraitBinding.h"
+
+#define UE_API UAFANIMGRAPH_API
+
+class UAnimNextAnimationGraph;
+struct FAnimNextDataInterfacePayload;
+struct FAnimNextFactoryParams;
+
+namespace UE::UAF
+{
+	/**
+	 * IGraphFactory
+	 *
+	 * This interface allows traits to participate in graph building via FAnimGraphFactory
+	 */
+	struct IGraphFactory : ITraitInterface
+	{
+		DECLARE_ANIM_TRAIT_INTERFACE(IGraphFactory)
+
+		// Called to modify/extend factory parameters for a generated graph.
+		// Object used to generate the base graph supplied for context.
+		// Fills in InOutParams with appropriately-initialized additive traits.
+		UE_API virtual void GetFactoryParams(FExecutionContext& Context, const TTraitBinding<IGraphFactory>& Binding, const UObject* InObject, FAnimNextFactoryParams& InOutParams) const;
+
+		// If the supplied object is a UAnimNextAnimationGraph, returns it
+		// If not then uses the parameters (od the default parameters if they are not valid) combined with the current trait stack's collective calls
+		// to GetFactoryParams to build the graph procedurally
+		UE_API static const UAnimNextAnimationGraph* GetOrBuildGraph(FExecutionContext& Context, const FTraitBinding& Binding, const UObject* InObject, FAnimNextFactoryParams& InOutParams);
+
+#if WITH_EDITOR
+		UE_API virtual const FText& GetDisplayName() const override;
+		UE_API virtual const FText& GetDisplayShortName() const override;
+#endif // WITH_EDITOR
+	};
+
+	/**
+	 * Specialization for trait binding.
+	 */
+	template<>
+	struct TTraitBinding<IGraphFactory> : FTraitBinding
+	{
+		// @see IGraphFactory::GetFactoryParams
+		void GetFactoryParams(FExecutionContext& Context, const UObject* InObject, FAnimNextFactoryParams& InOutParams) const
+		{
+			GetInterface()->GetFactoryParams(Context, *this, InObject, InOutParams);
+		}
+
+	protected:
+		const IGraphFactory* GetInterface() const { return GetInterfaceTyped<IGraphFactory>(); }
+	};
+}
+
+#undef UE_API
