UAF: Moved/renamed AnimNext and AnimNextAnimGraph plugins

diff --git a/Engine/Plugins/Experimental/UAF/UAFAnimGraph/Source/UAFAnimGraphEditor/Private/Factory/AnimNextFactoryParamsDetails.cpp b/Engine/Plugins/Experimental/UAF/UAFAnimGraph/Source/UAFAnimGraphEditor/Private/Factory/AnimNextFactoryParamsDetails.cpp
new file mode 100644
--- /dev/null
+++ b/Engine/Plugins/Experimental/UAF/UAFAnimGraph/Source/UAFAnimGraphEditor/Private/Factory/AnimNextFactoryParamsDetails.cpp
@@ -0,0 +1,514 @@
+﻿// Copyright Epic Games, Inc. All Rights Reserved.
+
+#include "AnimNextFactoryParamsDetails.h"
+
+#include "AnimNextStyle.h"
+#include "DetailLayoutBuilder.h"
+#include "DetailWidgetRow.h"
+#include "IDetailChildrenBuilder.h"
+#include "IDetailPropertyRow.h"
+#include "IStructureDataProvider.h"
+#include "PropertyHandle.h"
+#include "StructViewerFilter.h"
+#include "StructViewerModule.h"
+#include "Algo/Compare.h"
+#include "Factory/AnimGraphFactory.h"
+#include "Factory/AnimNextFactoryParams.h"
+#include "Modules/ModuleManager.h"
+#include "Styling/StyleColors.h"
+#include "TraitCore/TraitRegistry.h"
+#include "Widgets/Images/SImage.h"
+#include "Widgets/Input/SCheckBox.h"
+#include "Widgets/Input/SComboButton.h"
+#include "Widgets/Layout/SBox.h"
+#include "Widgets/Text/STextBlock.h"
+#include "Graph/AnimNextAnimationGraph.h"
+
+#define LOCTEXT_NAMESPACE "AnimNextFactoryParamsDetails"
+
+namespace UE::UAF::Editor
+{
+
+class FInstancedStructProvider : public IStructureDataProvider
+{
+public:
+	FInstancedStructProvider() = default;
+	
+	explicit FInstancedStructProvider(const TSharedPtr<IPropertyHandle>& InStructProperty)
+		: StructProperty(InStructProperty)
+	{
+	}
+	
+	virtual ~FInstancedStructProvider() override {}
+
+	void Reset()
+	{
+		StructProperty = nullptr;
+	}
+	
+	virtual bool IsValid() const override
+	{
+		bool bHasValidData = false;
+		EnumerateInstances([&bHasValidData](const UScriptStruct* ScriptStruct, uint8* Memory, UPackage* Package)
+		{
+			if (ScriptStruct && Memory)
+			{
+				bHasValidData = true;
+				return false; // Stop
+			}
+			return true; // Continue
+		});
+
+		return bHasValidData;
+	}
+	
+	virtual const UStruct* GetBaseStructure() const override
+	{
+		// Taken from UClass::FindCommonBase
+		auto FindCommonBaseStruct = [](const UScriptStruct* StructA, const UScriptStruct* StructB)
+		{
+			const UScriptStruct* CommonBaseStruct = StructA;
+			while (CommonBaseStruct && StructB && !StructB->IsChildOf(CommonBaseStruct))
+			{
+				CommonBaseStruct = Cast<UScriptStruct>(CommonBaseStruct->GetSuperStruct());
+			}
+			return CommonBaseStruct;
+		};
+
+		const UScriptStruct* CommonStruct = nullptr;
+		EnumerateInstances([&CommonStruct, &FindCommonBaseStruct](const UScriptStruct* ScriptStruct, uint8* Memory, UPackage* Package)
+		{
+			if (ScriptStruct)
+			{
+				CommonStruct = FindCommonBaseStruct(ScriptStruct, CommonStruct);
+			}
+			return true; // Continue
+		});
+
+		return CommonStruct;
+	}
+
+	virtual void GetInstances(TArray<TSharedPtr<FStructOnScope>>& OutInstances, const UStruct* ExpectedBaseStructure) const override
+	{
+		// The returned instances need to be compatible with base structure.
+		// This function returns empty instances in case they are not compatible, with the idea that we have as many instances as we have outer objects.
+		EnumerateInstances([&OutInstances, ExpectedBaseStructure](const UScriptStruct* ScriptStruct, uint8* Memory, UPackage* Package)
+		{
+			TSharedPtr<FStructOnScope> Result;
+			
+			if (ExpectedBaseStructure && ScriptStruct && ScriptStruct->IsChildOf(ExpectedBaseStructure))
+			{
+				Result = MakeShared<FStructOnScope>(ScriptStruct, Memory);
+				Result->SetPackage(Package);
+			}
+
+			OutInstances.Add(Result);
+
+			return true; // Continue
+		});
+	}
+
+	virtual bool IsPropertyIndirection() const override
+	{
+		return true;
+	}
+
+	virtual uint8* GetValueBaseAddress(uint8* ParentValueAddress, const UStruct* ExpectedBaseStructure) const override
+	{
+		if (!ParentValueAddress)
+		{
+			return nullptr;
+		}
+
+		FInstancedStruct& InstancedStruct = *reinterpret_cast<FInstancedStruct*>(ParentValueAddress);
+		if (ExpectedBaseStructure && InstancedStruct.GetScriptStruct() && InstancedStruct.GetScriptStruct()->IsChildOf(ExpectedBaseStructure))
+		{
+			return InstancedStruct.GetMutableMemory();
+		}
+		
+		return nullptr;
+	}
+	
+protected:
+
+	void EnumerateInstances(TFunctionRef<bool(const UScriptStruct* ScriptStruct, uint8* Memory, UPackage* Package)> InFunc) const
+	{
+		if (!StructProperty.IsValid() || !StructProperty->IsValidHandle())
+		{
+			return;
+		}
+		
+		TArray<UPackage*> Packages;
+		StructProperty->GetOuterPackages(Packages);
+
+		StructProperty->EnumerateRawData([&InFunc, &Packages](void* RawData, const int32 DataIndex, const int32 /*NumDatas*/)
+		{
+			const UScriptStruct* ScriptStruct = nullptr;
+			uint8* Memory = nullptr;
+			UPackage* Package = nullptr;
+			if (FInstancedStruct* InstancedStruct = static_cast<FInstancedStruct*>(RawData))
+			{
+				ScriptStruct = InstancedStruct->GetScriptStruct();
+				Memory = InstancedStruct->GetMutableMemory();
+				if (ensureMsgf(Packages.IsValidIndex(DataIndex), TEXT("Expecting packges and raw data to match.")))
+				{
+					Package = Packages[DataIndex];
+				}
+			}
+
+			return InFunc(ScriptStruct, Memory, Package);
+		});
+	}
+	
+	TSharedPtr<IPropertyHandle> StructProperty;
+};
+
+void FAnimNextFactoryParamsDetails::CustomizeHeader(TSharedRef<IPropertyHandle> InPropertyHandle, FDetailWidgetRow& InHeaderRow, IPropertyTypeCustomizationUtils& InCustomizationUtils)
+{
+}
+
+void FAnimNextFactoryParamsDetails::CustomizeChildren(TSharedRef<IPropertyHandle> InPropertyHandle, IDetailChildrenBuilder& InChildBuilder, IPropertyTypeCustomizationUtils& InCustomizationUtils)
+{
+	TSharedPtr<IPropertyHandle> BuilderHandle = InPropertyHandle->GetChildHandle(GET_MEMBER_NAME_CHECKED(FAnimNextFactoryParams, Builder));
+	BuilderHandle->MarkHiddenByCustomization();
+
+	TSharedPtr<IPropertyHandle> StacksHandle = BuilderHandle->GetChildHandle(GET_MEMBER_NAME_CHECKED(FAnimNextSimpleAnimGraphBuilder, Stacks));
+	TSharedPtr<IPropertyHandle> StackHandle = StacksHandle->GetChildHandle(0);
+	if (StackHandle.IsValid())
+	{
+		TraitStructsHandle = StackHandle->GetChildHandle(GET_MEMBER_NAME_CHECKED(FAnimNextSimpleAnimGraphBuilderTraitStackDesc, TraitStructs));
+	}
+	
+	bool bIsInvalidOrAnimGraph = false; 
+
+	// Find property for the factory source & common default params set
+	TOptional<FAnimNextFactoryParams> CommonDefaults;
+	FString FactorySourcePropertyName = InPropertyHandle->GetMetaData("FactorySourceProperty");
+	TSharedPtr<IPropertyHandle> FactorySourceHandle = InPropertyHandle->GetParentHandle()->GetChildHandle(FName(*FactorySourcePropertyName));
+	FactorySourceHandle->SetOnPropertyValueChanged(FSimpleDelegate::CreateLambda([WeakPropertyHandle = TWeakPtr<IPropertyHandle>(InPropertyHandle)]()
+	{
+		if (TSharedPtr<IPropertyHandle> PinnedPropertyHandle = WeakPropertyHandle.Pin())
+		{
+			PinnedPropertyHandle->RequestRebuildChildren();
+		}
+	}));
+	
+	if (FactorySourceHandle.IsValid() && FactorySourceHandle->GetProperty()->IsA<FObjectProperty>())
+	{
+		FactorySourceHandle->EnumerateRawData([&CommonDefaults, &bIsInvalidOrAnimGraph](void* RawData, const int32 DataIndex, const int32 NumDatas)
+		{
+			UObject* Object = *(static_cast<UObject**>(RawData));
+		
+			if (!CommonDefaults.IsSet())
+			{
+				if (Object == nullptr || Object->IsA<UAnimNextAnimationGraph>())
+				{
+					// TODO: At the moment we dont have 'details panel pin overrides' for assets, so we reset the params to not be confusing
+					// In the end we will want the params to also provide overrides like we supply to RunGraph nodes
+					CommonDefaults = FAnimNextFactoryParams();
+					bIsInvalidOrAnimGraph = true;
+				}
+				else
+				{
+					CommonDefaults = FAnimGraphFactory::GetDefaultParamsForObject(Object);
+				}
+			}
+			else
+			{
+				const FAnimNextFactoryParams& DefaultParams = FAnimGraphFactory::GetDefaultParamsForObject(Object);
+				if (!Algo::Compare(DefaultParams.Builder.Stacks[0].TraitStructs, CommonDefaults->Builder.Stacks[0].TraitStructs))
+				{
+					CommonDefaults.Reset();
+					return false;
+				}
+			}
+			return true;
+		});
+	}
+
+	TSet<const UScriptStruct*> AddedStructs;
+
+	if (TraitStructsHandle.IsValid())
+	{
+		// Find common 'current' set
+		TSet<const UScriptStruct*> CommonStructs;
+		TraitStructsHandle->EnumerateRawData([&CommonStructs](void* RawData, const int32 DataIndex, const int32 NumDatas)
+		{
+			TArray<FInstancedStruct>* InitializePayloads = static_cast<TArray<FInstancedStruct>*>(RawData);
+			if (CommonStructs.IsEmpty())
+			{
+				for (const FInstancedStruct& InitializePayload : *InitializePayloads)
+				{
+					CommonStructs.Add(InitializePayload.GetScriptStruct());
+				}
+			}
+			else
+			{
+				for (auto It = CommonStructs.CreateIterator(); It; ++It)
+				{
+					const UScriptStruct* CommonStruct = *It;
+					if (!InitializePayloads->ContainsByPredicate([CommonStruct](const FInstancedStruct& InInitializePayload){ return CommonStruct == InInitializePayload.GetScriptStruct(); }))
+					{
+						It.RemoveCurrent();
+					}
+				}
+			}
+			return true;
+		});
+
+		// No common structs, cannot edit (all struct types for edited objects are mutually exclusive)
+		if (CommonStructs.IsEmpty())
+		{
+			return;
+		}
+
+		// Add entries for instanced structs that are common to all objects and have a common type at this index
+		uint32 NumChildren = 0;
+		if (TraitStructsHandle->GetNumChildren(NumChildren) == FPropertyAccess::Success)
+		{
+			for (uint32 ChildIndex = 0; ChildIndex < NumChildren; ++ChildIndex)
+			{
+				TSharedRef<IPropertyHandle> InstancedStructHandle = TraitStructsHandle->GetChildHandle(ChildIndex).ToSharedRef();
+				TOptional<const UScriptStruct*> CommonType;
+				InstancedStructHandle->EnumerateRawData([&CommonType](void* RawData, const int32 DataIndex, const int32 NumDatas)
+				{
+					FInstancedStruct* InstancedStruct = static_cast<FInstancedStruct*>(RawData);
+					if (!CommonType.IsSet())
+					{
+						CommonType = InstancedStruct->GetScriptStruct();
+					}
+					else if (CommonType.GetValue() != InstancedStruct->GetScriptStruct())
+					{
+						CommonType.Reset();
+						return false;
+					}
+					return true;
+				});
+
+				if (CommonType.IsSet() && CommonStructs.Contains(CommonType.GetValue()))
+				{
+					AddedStructs.Add(CommonType.GetValue());
+
+					const FName UniqueName(InstancedStructHandle->GetPropertyPath());
+					IDetailPropertyRow* Row = InChildBuilder.AddChildStructure(InstancedStructHandle, MakeShared<FInstancedStructProvider>(InstancedStructHandle), UniqueName);
+					FDetailWidgetRow& CustomWidget = Row->CustomWidget(true);
+
+					CustomWidget
+					.NameContent()
+					[
+						SNew(SBorder)
+						.Padding(FMargin(6.0f, 2.0f))
+						.BorderImage(FAnimNextStyle::Get().GetBrush("AnimNext.TraitDetails.Background"))
+						[
+							SNew(STextBlock)
+							.TextStyle(FAnimNextStyle::Get(), "AnimNext.TraitDetails.Label")
+							.ColorAndOpacity(FStyleColors::Foreground)
+							.Text(CommonType.GetValue()->GetDisplayNameText())
+						]
+					];
+
+					const FTraitRegistry& TraitRegistry = FTraitRegistry::Get();
+					const FTrait* Trait = TraitRegistry.Find(CommonType.GetValue());
+					if (Trait)
+					{
+						// If this is an additive trait, allow it to be removed
+						if (Trait->GetTraitMode() == ETraitMode::Additive)
+						{
+							CustomWidget
+							.ValueContent()
+							[
+								SNew(SCheckBox)
+								.ToolTipText(LOCTEXT("EnableOptionalTrait", "Enable/disable this optional trait"))
+								.IsChecked(true)
+								.OnCheckStateChanged_Lambda([ChildIndex, WeakTraitStructsHandle = TWeakPtr<IPropertyHandle>(TraitStructsHandle), WeakPropertyHandle = TWeakPtr<IPropertyHandle>(InPropertyHandle)](ECheckBoxState InState)
+								{
+									TSharedPtr<IPropertyHandle> PinnedTraitStructsHandle = WeakTraitStructsHandle.Pin();
+									TSharedPtr<IPropertyHandle> PinnedPropertyHandle = WeakPropertyHandle.Pin();
+									if (PinnedTraitStructsHandle.IsValid() && PinnedPropertyHandle.IsValid())
+									{
+										PinnedTraitStructsHandle->AsArray()->DeleteItem(ChildIndex);
+										PinnedPropertyHandle->RequestRebuildChildren();
+									}
+								})
+							];
+						}
+						else if (FactorySourceHandle.IsValid())
+						{
+							// If its a base trait, hide any child property that has the appropriate metadata (if we have a factory source)
+							uint32 NumStructProperties = 0;
+							TSharedPtr<IPropertyHandle> StructHandle = InstancedStructHandle->GetChildHandle(0);
+							if (StructHandle->GetNumChildren(NumStructProperties) == FPropertyAccess::Success)
+							{
+								for (uint32 StructPropertyIndex = 0; StructPropertyIndex < NumStructProperties; ++StructPropertyIndex)
+								{
+									TSharedPtr<IPropertyHandle> StructPropertyHandle = StructHandle->GetChildHandle(StructPropertyIndex);
+									if (StructPropertyHandle->HasMetaData("FactorySource"))
+									{
+										StructPropertyHandle->MarkHiddenByCustomization();
+									}
+								}
+							}
+						}
+					}
+				}
+			}
+		}
+	}
+
+	if (CommonDefaults.IsSet() && CommonDefaults.GetValue().Builder.GetNumStacks() > 0)
+	{
+		// Now add any entries for defaults we didnt find above to allow them to be added
+		for (const TInstancedStruct<FAnimNextTraitSharedData>& DefaultStruct : CommonDefaults.GetValue().Builder.Stacks[0].TraitStructs)
+		{
+			if (!AddedStructs.Contains(DefaultStruct.GetScriptStruct()))
+			{
+				InChildBuilder.AddCustomRow(DefaultStruct.GetScriptStruct()->GetDisplayNameText())
+				.NameContent()
+				[
+					SNew(SBorder)
+					.Padding(FMargin(6.0f, 2.0f))
+					.BorderImage(FAnimNextStyle::Get().GetBrush("AnimNext.TraitDetails.Background"))
+					[
+						SNew(STextBlock)
+						.TextStyle(FAnimNextStyle::Get(), "AnimNext.TraitDetails.Label")
+						.ColorAndOpacity(FStyleColors::Foreground)
+						.Text(DefaultStruct.GetScriptStruct()->GetDisplayNameText())
+					]
+				]
+				.ValueContent()
+				[
+					SNew(SCheckBox)
+					.ToolTipText(LOCTEXT("EnableOptionalTrait", "Enable/disable this optional trait"))
+					.IsChecked(false)
+					.OnCheckStateChanged_Lambda([DefaultStruct, WeakTraitStructsHandle = TWeakPtr<IPropertyHandle>(TraitStructsHandle), WeakPropertyHandle = TWeakPtr<IPropertyHandle>(InPropertyHandle)](ECheckBoxState InState)
+					{
+						TSharedPtr<IPropertyHandle> PinnedTraitStructsHandle = WeakTraitStructsHandle.Pin();
+						TSharedPtr<IPropertyHandle> PinnedPropertyHandle = WeakPropertyHandle.Pin();
+						if (PinnedTraitStructsHandle.IsValid() && PinnedPropertyHandle.IsValid())
+						{
+							TSharedPtr<IPropertyHandleArray> ArrayHandle = PinnedTraitStructsHandle->AsArray();
+							FPropertyHandleItemAddResult Result = ArrayHandle->AddItem();
+							TSharedPtr<IPropertyHandle> NewEntryHandle = PinnedTraitStructsHandle->GetChildHandle(Result.GetIndex());
+							NewEntryHandle->EnumerateRawData([&DefaultStruct](void* RawData, const int32 DataIndex, const int32 NumDatas)
+							{
+								TInstancedStruct<FAnimNextTraitSharedData>* InstancedStruct = static_cast<TInstancedStruct<FAnimNextTraitSharedData>*>(RawData);
+								*InstancedStruct = DefaultStruct;
+								return true;
+							});
+							PinnedPropertyHandle->RequestRebuildChildren();
+						}
+					})
+				];
+			}
+		}
+	}
+
+	if (!bIsInvalidOrAnimGraph)
+	{
+		// Add an 'add' widget for any custom behaviors users want to opt into
+		InChildBuilder.AddCustomRow(LOCTEXT("AddTraitWidget", "Add Trait"))
+		.NameContent()
+		[
+			SNew(SComboButton)
+			.ToolTipText(LOCTEXT("AddTraitTooltip", "Add a trait to modify how this asset is played"))
+			.HasDownArrow(false)
+			.ComboButtonStyle(FAnimNextStyle::Get(), "AnimNext.TraitDetails.AddButton")
+			.ButtonContent()
+			[
+				SNew(SHorizontalBox)
+				+SHorizontalBox::Slot()
+				.VAlign(VAlign_Center)
+				.AutoWidth()
+				.Padding(0.0f, 0.0f, 3.0f, 0.0f)
+				[
+					SNew(SImage)
+					.Image(FAppStyle::Get().GetBrush("Plus"))
+				]
+				+SHorizontalBox::Slot()
+				.VAlign(VAlign_Center)
+				.AutoWidth()
+				.Padding(0.0f, 2.0f)
+				[
+					SNew(STextBlock)
+					.TextStyle(FAnimNextStyle::Get(), "AnimNext.TraitDetails.Label")
+					.ColorAndOpacity(FStyleColors::Foreground)
+					.Text(LOCTEXT("AddTraitButton", "Add"))
+				]
+			]
+			.OnGetMenuContent_Lambda([WeakTraitStructsHandle = TWeakPtr<IPropertyHandle>(TraitStructsHandle), WeakPropertyHandle = TWeakPtr<IPropertyHandle>(InPropertyHandle)]()
+			{
+				FStructViewerModule& StructViewerModule = FModuleManager::LoadModuleChecked<FStructViewerModule>("StructViewer");
+
+				class FStructFilter : public IStructViewerFilter
+				{
+				public:
+					virtual bool IsStructAllowed(const FStructViewerInitializationOptions& InInitOptions, const UScriptStruct* InStruct, TSharedRef<FStructViewerFilterFuncs> InFilterFuncs) override
+					{
+						if (InStruct->HasMetaData(TEXT("Hidden")))
+						{
+							return false;
+						}
+
+						if (!InStruct->IsChildOf(FAnimNextTraitSharedData::StaticStruct()))
+						{
+							return false;
+						}
+						
+						const FTraitRegistry& TraitRegistry = FTraitRegistry::Get();
+						const FTrait* Trait = TraitRegistry.Find(InStruct);
+						return Trait && Trait->GetTraitMode() == ETraitMode::Additive;
+					}
+
+					virtual bool IsUnloadedStructAllowed(const FStructViewerInitializationOptions& InInitOptions, const FSoftObjectPath& InStructPath, TSharedRef<class FStructViewerFilterFuncs> InFilterFuncs)
+					{
+						return false;
+					};
+				};
+
+				FStructViewerInitializationOptions Options;
+				Options.StructFilter = MakeShared<FStructFilter>();
+				Options.Mode = EStructViewerMode::StructPicker;
+				Options.NameTypeToDisplay = EStructViewerNameTypeToDisplay::DisplayName;
+				Options.bShowNoneOption = false;
+
+				return
+					SNew(SBox)
+					.WidthOverride(400.0f)
+					.HeightOverride(400.0f)
+					[
+						StructViewerModule.CreateStructViewer(Options, FOnStructPicked::CreateLambda([WeakPropertyHandle](const UScriptStruct* InStruct)
+						{
+							TSharedPtr<IPropertyHandle> PinnedPropertyHandle = WeakPropertyHandle.Pin();
+							if (PinnedPropertyHandle.IsValid())
+							{
+								TSharedPtr<IPropertyHandle> BuilderHandle = PinnedPropertyHandle->GetChildHandle(GET_MEMBER_NAME_CHECKED(FAnimNextFactoryParams, Builder));
+								TSharedPtr<IPropertyHandle> StacksHandle = BuilderHandle->GetChildHandle(GET_MEMBER_NAME_CHECKED(FAnimNextSimpleAnimGraphBuilder, Stacks));
+								TSharedPtr<IPropertyHandle> StackHandle = StacksHandle->GetChildHandle(0);
+								if (!StackHandle.IsValid())
+								{
+									StacksHandle->AsArray()->AddItem();
+									StackHandle = StacksHandle->GetChildHandle(0);
+								}
+
+								TSharedPtr<IPropertyHandle> TraitStructsHandle = StackHandle->GetChildHandle(GET_MEMBER_NAME_CHECKED(FAnimNextSimpleAnimGraphBuilderTraitStackDesc, TraitStructs));
+								TSharedPtr<IPropertyHandleArray> ArrayHandle = TraitStructsHandle->AsArray();
+								FPropertyHandleItemAddResult Result = ArrayHandle->AddItem();
+								TSharedPtr<IPropertyHandle> NewEntryHandle = TraitStructsHandle->GetChildHandle(Result.GetIndex());
+								NewEntryHandle->EnumerateRawData([InStruct](void* RawData, const int32 DataIndex, const int32 NumDatas)
+								{
+									FInstancedStruct* InstancedStruct = static_cast<FInstancedStruct*>(RawData);
+									InstancedStruct->InitializeAs(InStruct);
+									return true;
+								});
+								PinnedPropertyHandle->RequestRebuildChildren();
+							}
+						}))
+					];
+			})
+		];
+	}
+}
+
+}
+
+#undef LOCTEXT_NAMESPACE
\ No newline at end of file
diff --git a/Engine/Plugins/Experimental/UAF/UAFAnimGraph/Source/UAFAnimGraphEditor/Private/Graph/AnimNextAnimGraphEntryAssetDefinitions.cpp b/Engine/Plugins/Experimental/UAF/UAFAnimGraph/Source/UAFAnimGraphEditor/Private/Graph/AnimNextAnimGraphEntryAssetDefinitions.cpp
new file mode 100644
--- /dev/null
+++ b/Engine/Plugins/Experimental/UAF/UAFAnimGraph/Source/UAFAnimGraphEditor/Private/Graph/AnimNextAnimGraphEntryAssetDefinitions.cpp
@@ -0,0 +1,13 @@
+﻿// Copyright Epic Games, Inc. All Rights Reserved.
+
+#include "AnimNextAnimGraphEntryAssetDefinitions.h"
+
+#define LOCTEXT_NAMESPACE "AnimNextAnimGraphEntryAssetDefinitions"
+
+FText UAssetDefinition_AnimNextAnimationGraphEntry::GetObjectDisplayNameText(UObject* Object) const
+{
+	UAnimNextRigVMAssetEntry* Parameter = CastChecked<UAnimNextRigVMAssetEntry>(Object);
+	return Parameter->GetDisplayName();
+}
+
+#undef LOCTEXT_NAMESPACE
\ No newline at end of file
diff --git a/Engine/Plugins/Experimental/UAF/UAFAnimGraph/Source/UAFAnimGraphEditor/Private/Graph/AnimNextAnimationGraphFactory.cpp b/Engine/Plugins/Experimental/UAF/UAFAnimGraph/Source/UAFAnimGraphEditor/Private/Graph/AnimNextAnimationGraphFactory.cpp
new file mode 100644
--- /dev/null
+++ b/Engine/Plugins/Experimental/UAF/UAFAnimGraph/Source/UAFAnimGraphEditor/Private/Graph/AnimNextAnimationGraphFactory.cpp
@@ -0,0 +1,49 @@
+﻿// Copyright Epic Games, Inc. All Rights Reserved.
+
+#include "Graph/AnimNextAnimationGraphFactory.h"
+#include "Graph/AnimNextAnimationGraph.h"
+#include "Graph/AnimNextAnimationGraph_EditorData.h"
+#include "UncookedOnlyUtils.h"
+#include "Entries/AnimNextAnimationGraphEntry.h"
+#include "Graph/RigUnit_AnimNextGraphRoot.h"
+
+UAnimNextAnimationGraphFactory::UAnimNextAnimationGraphFactory()
+{
+	bCreateNew = true;
+	bEditAfterNew = true;
+	SupportedClass = UAnimNextAnimationGraph::StaticClass();
+}
+
+bool UAnimNextAnimationGraphFactory::ConfigureProperties()
+{
+	return true;
+}
+
+UObject* UAnimNextAnimationGraphFactory::FactoryCreateNew(UClass* Class, UObject* InParent, FName Name, EObjectFlags Flags, UObject* Context, FFeedbackContext* Warn, FName CallingContext)
+{
+	EObjectFlags FlagsToUse = Flags | RF_Public | RF_Standalone | RF_Transactional | RF_LoadCompleted;
+	if(InParent == GetTransientPackage())
+	{
+		FlagsToUse &= ~RF_Standalone;
+	}
+
+	UAnimNextAnimationGraph* NewAnimationGraph = NewObject<UAnimNextAnimationGraph>(InParent, Class, Name, FlagsToUse);
+
+	// Create internal editor data
+	UAnimNextAnimationGraph_EditorData* EditorData = NewObject<UAnimNextAnimationGraph_EditorData>(NewAnimationGraph, TEXT("EditorData"), RF_Transactional);
+	NewAnimationGraph->EditorData = EditorData;
+	EditorData->bUsesExternalPackages = false;
+	EditorData->Initialize(/*bRecompileVM*/false);
+
+	// Add a single internal graph entry
+	UAnimNextAnimationGraphEntry* NewGraph = EditorData->AddAnimationGraph(FRigUnit_AnimNextGraphRoot::DefaultEntryPoint, false, false);
+
+	// Hide in the outliner - the asset itself represents the graph to users
+	NewGraph->SetHiddenInOutliner(true);
+
+	// Compile the initial skeleton
+	EditorData->RecompileVM();
+	check(!EditorData->bErrorsDuringCompilation);
+
+	return NewAnimationGraph;
+}
\ No newline at end of file
diff --git a/Engine/Plugins/Experimental/UAF/UAFAnimGraph/Source/UAFAnimGraphEditor/Private/Graph/AnimNextAnimationGraphMenuExtensions.cpp b/Engine/Plugins/Experimental/UAF/UAFAnimGraph/Source/UAFAnimGraphEditor/Private/Graph/AnimNextAnimationGraphMenuExtensions.cpp
new file mode 100644
--- /dev/null
+++ b/Engine/Plugins/Experimental/UAF/UAFAnimGraph/Source/UAFAnimGraphEditor/Private/Graph/AnimNextAnimationGraphMenuExtensions.cpp
@@ -0,0 +1,329 @@
+// Copyright Epic Games, Inc. All Rights Reserved.
+
+#include "AnimNextAnimationGraphMenuExtensions.h"
+
+#include "AnimGraphUncookedOnlyUtils.h"
+#include "AnimNextEdGraphNode.h"
+#include "ContentBrowserModule.h"
+#include "IContentBrowserSingleton.h"
+#include "ToolMenus.h"
+#include "UncookedOnlyUtils.h"
+#include "EdGraph/EdGraphNode.h"
+#include "Framework/Application/SlateApplication.h"
+#include "Graph/RigDecorator_AnimNextCppTrait.h"
+#include "Graph/RigUnit_AnimNextRunAnimationGraph_v1.h"
+#include "Graph/RigUnit_AnimNextRunAnimationGraph_v2.h"
+#include "Graph/RigVMTrait_AnimNextPublicVariables.h"
+#include "RigVMModel/RigVMController.h"
+#include "RigVMModel/Nodes/RigVMUnitNode.h"
+#include "TraitCore/TraitRegistry.h"
+#include "AnimNextController.h"
+#include "AnimNextTraitStackUnitNode.h"
+#include "AnimNextRigVMAssetEditorData.h"
+#include "RigVMModel/RigVMClient.h"
+#include "Variables/AnimNextSharedVariables.h"
+
+#define LOCTEXT_NAMESPACE "FAnimationGraphMenuExtensions"
+
+namespace UE::UAF::Editor::Private
+{
+	static const FLazyName VariablesTraitBaseName = TEXT("Variables");
+}
+
+namespace UE::UAF::Editor
+{
+
+void FAnimationGraphMenuExtensions::RegisterMenus()
+{
+	FToolMenuOwnerScoped OwnerScoped(TEXT("FAnimNextAnimationGraphItemDetails"));
+	UToolMenu* Menu = UToolMenus::Get()->ExtendMenu("GraphEditor.GraphNodeContextMenu.AnimNextEdGraphNode");
+	if (Menu == nullptr)
+	{
+		return;
+	}
+
+	Menu->AddDynamicSection(TEXT("AnimNextEdGraphNode"), FNewToolMenuDelegate::CreateLambda([](UToolMenu* InMenu)
+	{
+		UGraphNodeContextMenuContext* Context = InMenu->FindContext<UGraphNodeContextMenuContext>();
+		if(Context == nullptr)
+		{
+			return;
+		}
+
+		const UAnimNextEdGraphNode* AnimNextEdGraphNode = Cast<UAnimNextEdGraphNode>(Context->Node);
+		if(AnimNextEdGraphNode == nullptr)
+		{
+			return;
+		}
+
+		URigVMNode* ModelNode = AnimNextEdGraphNode->GetModelNode();
+		if(ModelNode == nullptr)
+		{
+			return;
+		}
+
+		auto IsRunGraphNode = [](URigVMNode* InModelNode)
+		{
+			if (const URigVMUnitNode* VMNode = Cast<URigVMUnitNode>(InModelNode))
+			{
+				const UScriptStruct* ScriptStruct = VMNode->GetScriptStruct();
+				return ScriptStruct == FRigUnit_AnimNextRunAnimationGraph_v1::StaticStruct() || ScriptStruct == FRigUnit_AnimNextRunAnimationGraph_v2::StaticStruct();
+			}
+
+			return false;
+		};
+		
+		if (UncookedOnly::FAnimGraphUtils::IsTraitStackNode(ModelNode))
+		{
+			FToolMenuSection& Section = InMenu->AddSection("AnimNextTraitNodeActions", LOCTEXT("AnimNextTraitNodeActionsMenuHeader", "Traits"));
+
+			auto BuildAddTraitContextMenu = [AnimNextEdGraphNode, ModelNode](UToolMenu* InSubMenu)
+			{
+				const FTraitRegistry& TraitRegistry = FTraitRegistry::Get();
+				TArray<const FTrait*> Traits = TraitRegistry.GetTraits();
+
+				URigVMController* VMController = AnimNextEdGraphNode->GetController();
+
+				const UScriptStruct* CppDecoratorStruct = FRigDecorator_AnimNextCppDecorator::StaticStruct();
+
+				for (const FTrait* Trait : Traits)
+				{
+					UScriptStruct* ScriptStruct = Trait->GetTraitSharedDataStruct();
+			
+					FString DefaultValue;
+					{
+						const FRigDecorator_AnimNextCppDecorator DefaultCppDecoratorStructInstance;
+						FRigDecorator_AnimNextCppDecorator CppDecoratorStructInstance;
+						CppDecoratorStructInstance.DecoratorSharedDataStruct = ScriptStruct;
+			
+						if (!CppDecoratorStructInstance.CanBeAddedToNode(ModelNode, nullptr))
+						{
+							continue;	// This trait isn't supported on this node
+						}
+
+						FRigDecorator_AnimNextCppDecorator::StaticStruct()->ExportText(DefaultValue, &CppDecoratorStructInstance, &DefaultCppDecoratorStructInstance, nullptr, PPF_SerializedAsImportText, nullptr);
+					}
+			
+					FString DisplayNameMetadata;
+					ScriptStruct->GetStringMetaDataHierarchical(FRigVMStruct::DisplayNameMetaName, &DisplayNameMetadata);
+					const FString DisplayName = DisplayNameMetadata.IsEmpty() ? Trait->GetTraitName() : DisplayNameMetadata;
+					const FText ToolTip = ScriptStruct->GetToolTipText();
+			
+					FToolMenuEntry TraitEntry = FToolMenuEntry::InitMenuEntry(
+						*Trait->GetTraitName(),
+						FText::FromString(DisplayName),
+						ToolTip,
+						FSlateIcon(),
+						FUIAction(FExecuteAction::CreateLambda(
+							[Trait, VMController, ModelNode, CppDecoratorStruct, DefaultValue, DisplayName]()
+							{
+								VMController->AddTrait(
+									ModelNode->GetFName(),
+									*CppDecoratorStruct->GetPathName(),
+									*DisplayName,
+									DefaultValue, INDEX_NONE, true, true);
+							})
+						)
+					);
+
+					InSubMenu->AddMenuEntry(NAME_None, TraitEntry);
+				}
+			};
+			
+			Section.AddSubMenu(
+				TEXT("AddTraitMenu"),
+				LOCTEXT("AddTraitMenu", "Add Trait"),
+				LOCTEXT("AddTraitMenuTooltip", "Add the chosen trait to currently selected node"),
+				FNewToolMenuDelegate::CreateLambda(BuildAddTraitContextMenu));
+		}
+		else if(IsRunGraphNode(ModelNode))
+		{
+			FToolMenuSection& Section = InMenu->AddSection("AnimNextRunAnimGraphNodeActions", LOCTEXT("AnimNextAnimGraphNodeActionsMenuHeader", "Animation Graph"));
+	
+			URigVMController* VMController = AnimNextEdGraphNode->GetController();
+			URigVMPin* VMPin = Context->Pin != nullptr ? AnimNextEdGraphNode->FindModelPinFromGraphPin(Context->Pin) : nullptr;
+	
+			if(VMPin != nullptr && ModelNode->FindTrait(VMPin))
+			{
+				Section.AddMenuEntry(
+					TEXT("RemoveExposedVariables"),
+					LOCTEXT("RemoveExposedVariablesMenu", "Remove Exposed Variables"),
+					LOCTEXT("RemoveExposeVariablesMenuTooltip", "Remove the exposed variable trait from this node"),
+					FSlateIcon(),
+					FUIAction(
+						FExecuteAction::CreateLambda([VMController, ModelNode, Name = VMPin->GetFName()]()
+						{
+							VMController->RemoveTrait(ModelNode->GetFName(), Name, true, true);
+						})
+					));
+			}
+			else
+			{
+				auto BuildExposeVariablesContextMenu = [AnimNextEdGraphNode, ModelNode](UToolMenu* InSubMenu)
+				{
+					URigVMController* VMController = AnimNextEdGraphNode->GetController();
+
+					FContentBrowserModule& ContentBrowserModule = FModuleManager::Get().LoadModuleChecked<FContentBrowserModule>(TEXT("ContentBrowser"));
+
+					FAssetPickerConfig AssetPickerConfig;
+					AssetPickerConfig.Filter.ClassPaths.Add(UAnimNextSharedVariables::StaticClass()->GetClassPathName());
+					AssetPickerConfig.Filter.bRecursiveClasses = true;
+					AssetPickerConfig.InitialAssetViewType = EAssetViewType::List;
+					AssetPickerConfig.AssetShowWarningText = LOCTEXT("NoAssetsWithPublicVariablesMessage", "No animation graphs with public variables found");
+					AssetPickerConfig.OnAssetSelected = FOnAssetSelected::CreateLambda([VMController, ModelNode](const FAssetData& InAssetData)
+					{
+						FSlateApplication::Get().DismissAllMenus();
+
+						FString DefaultValue;
+						FRigVMTrait_AnimNextPublicVariables DefaultTrait;
+						FRigVMTrait_AnimNextPublicVariables NewTrait;
+						UAnimNextSharedVariables* Asset = CastChecked<UAnimNextSharedVariables>(InAssetData.GetAsset());
+						NewTrait.InternalAsset = Asset;
+						const FInstancedPropertyBag& VariableDefaults = Asset->GetVariableDefaults();
+						TConstArrayView<FPropertyBagPropertyDesc> Descs = VariableDefaults.GetPropertyBagStruct()->GetPropertyDescs();
+						NewTrait.InternalVariableNames.Reserve(Descs.Num());
+						for(const FPropertyBagPropertyDesc& Desc : Descs)
+						{
+							if (Desc.CachedProperty->HasAnyPropertyFlags(CPF_NativeAccessSpecifierPublic))
+							{
+								NewTrait.InternalVariableNames.Add(Desc.Name);
+							}
+						}
+						FRigVMTrait_AnimNextPublicVariables::StaticStruct()->ExportText(DefaultValue, &NewTrait, &DefaultTrait, nullptr, PPF_SerializedAsImportText, nullptr);
+
+						const FName ValidTraitName = URigVMSchema::GetUniqueName(Private::VariablesTraitBaseName, [ModelNode](const FName& InName)
+						{
+							return ModelNode->FindPin(InName.ToString()) == nullptr;
+						}, false, false);
+
+						VMController->AddTrait(ModelNode->GetFName(), *FRigVMTrait_AnimNextPublicVariables::StaticStruct()->GetPathName(), ValidTraitName, DefaultValue, INDEX_NONE, true, true);
+					});
+
+					AssetPickerConfig.OnShouldFilterAsset = FOnShouldFilterAsset::CreateLambda([](const FAssetData& InAssetData)
+					{
+						// Filter to only show assets with public variables
+						FAnimNextAssetRegistryExports Exports;
+						if(UncookedOnly::FUtils::GetExportedVariablesForAsset(InAssetData, Exports))
+						{
+							for(const FAnimNextExport& Export : Exports.Exports)
+							{
+								if (const FAnimNextVariableDeclarationData* VariableDeclaration = Export.Data.GetPtr<FAnimNextVariableDeclarationData>())
+								{
+									if (EnumHasAnyFlags(static_cast<EAnimNextExportedVariableFlags>(VariableDeclaration->Flags),EAnimNextExportedVariableFlags::Public))
+									{
+										return false;
+									}					
+								}
+							}
+						}
+						return true;
+					});
+
+					FToolMenuEntry Entry = FToolMenuEntry::InitWidget(
+						TEXT("AnimationGraphPicker"),
+						SNew(SBox)
+						.WidthOverride(300.0f)
+						.HeightOverride(400.0f)
+						[
+							ContentBrowserModule.Get().CreateAssetPicker(AssetPickerConfig)
+						],
+						FText::GetEmpty(),
+						true,
+						false,
+						false,
+						LOCTEXT("AnimationGraphPickerTooltip", "Choose an animation graph with public variables to expose")
+					);
+
+					InSubMenu->AddMenuEntry(NAME_None, Entry);
+				};
+				
+				Section.AddSubMenu(
+					TEXT("ExposeVariables"),
+					LOCTEXT("ExposeVariablesMenu", "Expose Variables"),
+					LOCTEXT("ExposeVariablesMenuTooltip", "Expose the variables of a selected animation graph as pins on this node"),
+					FNewToolMenuDelegate::CreateLambda(BuildExposeVariablesContextMenu));
+			}
+		}
+
+		if (UAnimNextTraitStackUnitNode* AnimNextUnitNode = Cast<UAnimNextTraitStackUnitNode>(ModelNode))
+		{
+			AddManifestSection(InMenu, AnimNextEdGraphNode, ModelNode);
+		}
+	}));
+}
+
+void FAnimationGraphMenuExtensions::UnregisterMenus()
+{
+	if(UToolMenus* ToolMenus = UToolMenus::Get())
+	{
+		ToolMenus->UnregisterOwnerByName("FAnimNextAnimationGraphItemDetails");
+	}
+}
+
+void FAnimationGraphMenuExtensions::AddManifestSection(UToolMenu* InMenu, const UAnimNextEdGraphNode* AnimNextEdGraphNode, const URigVMNode* ModelNode)
+{
+	FToolMenuSection& Section = InMenu->AddSection("AnimNextManifestNodeActions", LOCTEXT("AnimNextManifestNodeActionsMenuHeader", "Manifest"));
+
+	const TWeakObjectPtr<const UAnimNextEdGraphNode> EdNodeWeak = AnimNextEdGraphNode;
+
+	if (!UE::UAF::UncookedOnly::FAnimGraphUtils::IsExposedToManifest(ModelNode))
+	{
+		Section.AddMenuEntry(
+			"AddTraitToManifest",
+			LOCTEXT("AddTraitToManifest", "Add to Manifest"),
+			LOCTEXT("AddTraitToManifest_Tooltip", "Adds this Trait Stack to the Manifest."),
+			FSlateIcon(),
+			FUIAction(FExecuteAction::CreateLambda([EdNodeWeak]()
+				{
+					if (const TStrongObjectPtr<const UAnimNextEdGraphNode> EdGraphNode = EdNodeWeak.Pin())
+					{
+						if (URigVMNode* ModelNode = EdGraphNode->GetModelNode())
+						{
+							UAnimNextController* Controller = CastChecked<UAnimNextController>(EdGraphNode->GetController());
+
+							if (ensure(UE::UAF::UncookedOnly::FAnimGraphUtils::IsExposedToManifest(ModelNode) == false))
+							{
+								FRigVMControllerCompileBracketScope CompileScope(Controller);
+
+								Controller->AddNodeToManifest(ModelNode, true, true);
+								
+								if (UAnimNextRigVMAssetEditorData* EditorData = EdGraphNode->GetGraph()->GetTypedOuter<UAnimNextRigVMAssetEditorData>())
+								{
+									UE::UAF::UncookedOnly::FAnimGraphUtils::RequestVMAutoRecompile(EditorData);  // required to force an asset tags update
+								}
+							}
+						}								
+					}
+				}))
+			);
+	}
+	else
+	{
+		Section.AddMenuEntry(
+			"RemoveTraitFromManifest",
+			LOCTEXT("RemoveTraitFromManifest", "Remove from Manifest"),
+			LOCTEXT("RemoveTraitFromManifest_Tooltip", "Removes this Trait Stack from the Manifest."),
+			FSlateIcon(),
+			FUIAction(FExecuteAction::CreateLambda([EdNodeWeak]()
+				{
+					if (const TStrongObjectPtr<const UAnimNextEdGraphNode> EdGraphNode = EdNodeWeak.Pin())
+					{
+						UAnimNextController* Controller = CastChecked<UAnimNextController>(EdGraphNode->GetController());
+
+						FRigVMControllerCompileBracketScope CompileScope(Controller);
+
+						Controller->RemoveNodeFromManifest(EdGraphNode->GetModelNode(), true, true);
+						
+						if (UAnimNextRigVMAssetEditorData* EditorData = EdGraphNode->GetGraph()->GetTypedOuter<UAnimNextRigVMAssetEditorData>())
+						{
+							UE::UAF::UncookedOnly::FAnimGraphUtils::RequestVMAutoRecompile(EditorData); // required to force an asset tags update
+						}
+					}
+				}))
+		);
+	}
+}
+
+}
+
+#undef LOCTEXT_NAMESPACE
\ No newline at end of file
diff --git a/Engine/Plugins/Experimental/UAF/UAFAnimGraph/Source/UAFAnimGraphEditor/Private/Graph/AnimNextAnimationGraphMenuExtensions.h b/Engine/Plugins/Experimental/UAF/UAFAnimGraph/Source/UAFAnimGraphEditor/Private/Graph/AnimNextAnimationGraphMenuExtensions.h
new file mode 100644
--- /dev/null
+++ b/Engine/Plugins/Experimental/UAF/UAFAnimGraph/Source/UAFAnimGraphEditor/Private/Graph/AnimNextAnimationGraphMenuExtensions.h
@@ -0,0 +1,25 @@
+// Copyright Epic Games, Inc. All Rights Reserved.
+
+#pragma once
+
+#include "CoreMinimal.h"
+
+class UAnimNextEdGraphNode;
+class URigVMNode;
+class UToolMenu;
+struct FToolMenuSection;
+
+namespace UE::UAF::Editor
+{
+
+struct FAnimationGraphMenuExtensions
+{
+	static void RegisterMenus();
+	static void UnregisterMenus();
+
+private:
+	static void AddManifestSection(UToolMenu* InMenu, const UAnimNextEdGraphNode* AnimNextEdGraphNode, const URigVMNode* ModelNode);
+};
+
+}
+
