Run UnrealCodeFixup to add #include UE_INLINE_GENERATED_CPP_BY_NAME

diff --git a/Engine/Plugins/Experimental/UAF/UAFAnimGraph/Source/UAFAnimGraph/Private/Graph/AnimNextAnimGraph.cpp b/Engine/Plugins/Experimental/UAF/UAFAnimGraph/Source/UAFAnimGraph/Private/Graph/AnimNextAnimGraph.cpp
--- a/Engine/Plugins/Experimental/UAF/UAFAnimGraph/Source/UAFAnimGraph/Private/Graph/AnimNextAnimGraph.cpp
+++ b/Engine/Plugins/Experimental/UAF/UAFAnimGraph/Source/UAFAnimGraph/Private/Graph/AnimNextAnimGraph.cpp
@@ -3,6 +3,8 @@
 #include "Graph/AnimNextAnimGraph.h"
 #include "Graph/AnimNextGraphInstance.h"
 
+#include UE_INLINE_GENERATED_CPP_BY_NAME(AnimNextAnimGraph)
+
 bool FAnimNextAnimGraph::IsEqualForInjectionSiteChange(const FAnimNextAnimGraph& InOther) const
 {
 	return Asset == InOther.Asset &&
diff --git a/Engine/Plugins/Experimental/UAF/UAFAnimGraph/Source/UAFAnimGraph/Private/Graph/AnimNextGraphInstance.cpp b/Engine/Plugins/Experimental/UAF/UAFAnimGraph/Source/UAFAnimGraph/Private/Graph/AnimNextGraphInstance.cpp
--- a/Engine/Plugins/Experimental/UAF/UAFAnimGraph/Source/UAFAnimGraph/Private/Graph/AnimNextGraphInstance.cpp
+++ b/Engine/Plugins/Experimental/UAF/UAFAnimGraph/Source/UAFAnimGraph/Private/Graph/AnimNextGraphInstance.cpp
@@ -16,6 +16,8 @@
 #include "Module/AnimNextModule.h"
 #include "Module/AnimNextModuleInstance.h"
 
+#include UE_INLINE_GENERATED_CPP_BY_NAME(AnimNextGraphInstance)
+
 DEFINE_STAT(STAT_AnimNext_Graph_RigVM);
 
 FAnimNextGraphInstance::FAnimNextGraphInstance()
diff --git a/Engine/Plugins/Experimental/UAF/UAFAnimGraph/Source/UAFAnimGraph/Private/Graph/AnimNextModuleAnimGraphComponent.cpp b/Engine/Plugins/Experimental/UAF/UAFAnimGraph/Source/UAFAnimGraph/Private/Graph/AnimNextModuleAnimGraphComponent.cpp
--- a/Engine/Plugins/Experimental/UAF/UAFAnimGraph/Source/UAFAnimGraph/Private/Graph/AnimNextModuleAnimGraphComponent.cpp
+++ b/Engine/Plugins/Experimental/UAF/UAFAnimGraph/Source/UAFAnimGraph/Private/Graph/AnimNextModuleAnimGraphComponent.cpp
@@ -6,6 +6,8 @@
 #include "Graph/AnimNextAnimationGraph.h"
 #include "Module/AnimNextModuleInstance.h"
 
+#include UE_INLINE_GENERATED_CPP_BY_NAME(AnimNextModuleAnimGraphComponent)
+
 TWeakPtr<FAnimNextGraphInstance> FAnimNextModuleAnimGraphComponent::AllocateInstance(const UAnimNextAnimationGraph* InAnimationGraph, const TSharedPtr<const UE::UAF::FVariableOverridesCollection>& InOverrides)
 {
 	check(InAnimationGraph);
diff --git a/Engine/Plugins/Experimental/UAF/UAFAnimGraph/Source/UAFAnimGraph/Private/Graph/AnimNextSkeletalMeshComponentLibrary.cpp b/Engine/Plugins/Experimental/UAF/UAFAnimGraph/Source/UAFAnimGraph/Private/Graph/AnimNextSkeletalMeshComponentLibrary.cpp
--- a/Engine/Plugins/Experimental/UAF/UAFAnimGraph/Source/UAFAnimGraph/Private/Graph/AnimNextSkeletalMeshComponentLibrary.cpp
+++ b/Engine/Plugins/Experimental/UAF/UAFAnimGraph/Source/UAFAnimGraph/Private/Graph/AnimNextSkeletalMeshComponentLibrary.cpp
@@ -1,8 +1,10 @@
-﻿// Copyright Epic Games, Inc. All Rights Reserved.
+// Copyright Epic Games, Inc. All Rights Reserved.
 
 #include "AnimNextSkeletalMeshComponentLibrary.h"
 #include "Components/SkeletalMeshComponent.h"
 
+#include UE_INLINE_GENERATED_CPP_BY_NAME(AnimNextSkeletalMeshComponentLibrary)
+
 FAnimNextGraphReferencePose UAnimNextSkeletalMeshComponentLibrary::GetReferencePose(USkeletalMeshComponent* InComponent)
 {
 	UE::UAF::FDataHandle RefPoseHandle = UE::UAF::FDataRegistry::Get()->GetOrGenerateReferencePose(InComponent);
diff --git a/Engine/Plugins/Experimental/UAF/UAFAnimGraph/Source/UAFAnimGraph/Private/Graph/GC_GraphInstanceComponent.cpp b/Engine/Plugins/Experimental/UAF/UAFAnimGraph/Source/UAFAnimGraph/Private/Graph/GC_GraphInstanceComponent.cpp
--- a/Engine/Plugins/Experimental/UAF/UAFAnimGraph/Source/UAFAnimGraph/Private/Graph/GC_GraphInstanceComponent.cpp
+++ b/Engine/Plugins/Experimental/UAF/UAFAnimGraph/Source/UAFAnimGraph/Private/Graph/GC_GraphInstanceComponent.cpp
@@ -5,6 +5,8 @@
 #include "TraitCore/ExecutionContext.h"
 #include "TraitInterfaces/IGarbageCollection.h"
 
+#include UE_INLINE_GENERATED_CPP_BY_NAME(GC_GraphInstanceComponent)
+
 void FGCGraphInstanceComponent::Register(const UE::UAF::FWeakTraitPtr& InTraitPtr)
 {
 	TraitsWithReferences.Add(InTraitPtr);
diff --git a/Engine/Plugins/Experimental/UAF/UAFAnimGraph/Source/UAFAnimGraph/Private/Graph/RigUnit_AnimNextSetNotifyContext.cpp b/Engine/Plugins/Experimental/UAF/UAFAnimGraph/Source/UAFAnimGraph/Private/Graph/RigUnit_AnimNextSetNotifyContext.cpp
--- a/Engine/Plugins/Experimental/UAF/UAFAnimGraph/Source/UAFAnimGraph/Private/Graph/RigUnit_AnimNextSetNotifyContext.cpp
+++ b/Engine/Plugins/Experimental/UAF/UAFAnimGraph/Source/UAFAnimGraph/Private/Graph/RigUnit_AnimNextSetNotifyContext.cpp
@@ -6,10 +6,12 @@
 #include "Module/AnimNextModuleInstance.h"
 #include "Traits/NotifyDispatcher.h"
 
+#include UE_INLINE_GENERATED_CPP_BY_NAME(RigUnit_AnimNextSetNotifyContext)
+
 FRigUnit_AnimNextSetNotifyContext_Execute()
 {
 	FAnimNextModuleInstance& ModuleInstance = ExecuteContext.GetContextData<FAnimNextModuleContextData>().GetModuleInstance();
 	FAnimNextNotifyDispatcherComponent& NotifyDispatcher = ModuleInstance.GetComponent<FAnimNextNotifyDispatcherComponent>();
 	NotifyDispatcher.SkeletalMeshComponent = SkeletalMeshComponent;
 	NotifyDispatcher.NotifyQueue.PredictedLODLevel = SkeletalMeshComponent ? SkeletalMeshComponent->GetPredictedLODLevel() : 0;
-}
\ No newline at end of file
+}
diff --git a/Engine/Plugins/Experimental/UAF/UAFAnimGraph/Source/UAFAnimGraph/Private/Graph/RigUnit_GetPostProcessAnimation.cpp b/Engine/Plugins/Experimental/UAF/UAFAnimGraph/Source/UAFAnimGraph/Private/Graph/RigUnit_GetPostProcessAnimation.cpp
--- a/Engine/Plugins/Experimental/UAF/UAFAnimGraph/Source/UAFAnimGraph/Private/Graph/RigUnit_GetPostProcessAnimation.cpp
+++ b/Engine/Plugins/Experimental/UAF/UAFAnimGraph/Source/UAFAnimGraph/Private/Graph/RigUnit_GetPostProcessAnimation.cpp
@@ -8,6 +8,8 @@
 #include "Module/AnimNextModuleInstance.h"
 #include "Module/AnimNextSkeletalMeshComponentReferenceComponent.h"
 
+#include UE_INLINE_GENERATED_CPP_BY_NAME(RigUnit_GetPostProcessAnimation)
+
 FRigUnit_GetPostProcessAnimation_Execute()
 {
 	Graph = {};
@@ -38,4 +40,4 @@ FRigUnit_GetPostProcessAnimation_Execute()
 			bShouldEvaluate = SkeletalMesh->ShouldEvaluatePostProcessAnimGraph(LODLevel);
 		}
 	}
-}
\ No newline at end of file
+}
diff --git a/Engine/Plugins/Experimental/UAF/UAFAnimGraph/Source/UAFAnimGraph/Private/Graph/RigUnit_HasValidEvaluationModifier.cpp b/Engine/Plugins/Experimental/UAF/UAFAnimGraph/Source/UAFAnimGraph/Private/Graph/RigUnit_HasValidEvaluationModifier.cpp
--- a/Engine/Plugins/Experimental/UAF/UAFAnimGraph/Source/UAFAnimGraph/Private/Graph/RigUnit_HasValidEvaluationModifier.cpp
+++ b/Engine/Plugins/Experimental/UAF/UAFAnimGraph/Source/UAFAnimGraph/Private/Graph/RigUnit_HasValidEvaluationModifier.cpp
@@ -2,6 +2,8 @@
 
 #include "RigUnit_HasValidEvaluationModifier.h"
 
+#include UE_INLINE_GENERATED_CPP_BY_NAME(RigUnit_HasValidEvaluationModifier)
+
 FRigUnit_HasValidEvaluationModifier_Execute()
 {
 	bValue = InjectionData.HasValidEvaluationModifier();
diff --git a/Engine/Plugins/Experimental/UAF/UAFAnimGraph/Source/UAFAnimGraph/Private/Graph/SyncGroup_GraphInstanceComponent.cpp b/Engine/Plugins/Experimental/UAF/UAFAnimGraph/Source/UAFAnimGraph/Private/Graph/SyncGroup_GraphInstanceComponent.cpp
--- a/Engine/Plugins/Experimental/UAF/UAFAnimGraph/Source/UAFAnimGraph/Private/Graph/SyncGroup_GraphInstanceComponent.cpp
+++ b/Engine/Plugins/Experimental/UAF/UAFAnimGraph/Source/UAFAnimGraph/Private/Graph/SyncGroup_GraphInstanceComponent.cpp
@@ -15,6 +15,8 @@
 // Enabled by default in development builds
 // We log debug information using the visual logger
 // Use Rewind Debugger to record and replay logs in editor
+
+#include UE_INLINE_GENERATED_CPP_BY_NAME(SyncGroup_GraphInstanceComponent)
 #define UE_DEBUG_SYNC_GROUPS !(UE_BUILD_SHIPPING || UE_BUILD_TEST)
 
 // [Sync Group Details]
diff --git a/Engine/Plugins/Experimental/UAF/UAFAnimGraph/Source/UAFAnimGraph/Private/Graph/UAFGraphInstanceComponent.cpp b/Engine/Plugins/Experimental/UAF/UAFAnimGraph/Source/UAFAnimGraph/Private/Graph/UAFGraphInstanceComponent.cpp
--- a/Engine/Plugins/Experimental/UAF/UAFAnimGraph/Source/UAFAnimGraph/Private/Graph/UAFGraphInstanceComponent.cpp
+++ b/Engine/Plugins/Experimental/UAF/UAFAnimGraph/Source/UAFAnimGraph/Private/Graph/UAFGraphInstanceComponent.cpp
@@ -1,8 +1,10 @@
-﻿// Copyright Epic Games, Inc. All Rights Reserved.
+// Copyright Epic Games, Inc. All Rights Reserved.
 
 #include "Graph/UAFGraphInstanceComponent.h"
 #include "Graph/AnimNextGraphInstance.h"
 
+#include UE_INLINE_GENERATED_CPP_BY_NAME(UAFGraphInstanceComponent)
+
 FAnimNextGraphInstance& FUAFGraphInstanceComponent::GetGraphInstance()
 {
 	return static_cast<FAnimNextGraphInstance&>(*Instance);
@@ -11,4 +13,4 @@ FAnimNextGraphInstance& FUAFGraphInstanceComponent::GetGraphInstance()
 const FAnimNextGraphInstance& FUAFGraphInstanceComponent::GetGraphInstance() const
 {
 	return static_cast<const FAnimNextGraphInstance&>(*Instance);
-}
\ No newline at end of file
+}
diff --git a/Engine/Plugins/Experimental/UAF/UAFAnimGraph/Source/UAFAnimGraph/Private/GraphInterfaces/AnimNextNativeDataInterface_AnimSequencePlayer.cpp b/Engine/Plugins/Experimental/UAF/UAFAnimGraph/Source/UAFAnimGraph/Private/GraphInterfaces/AnimNextNativeDataInterface_AnimSequencePlayer.cpp
--- a/Engine/Plugins/Experimental/UAF/UAFAnimGraph/Source/UAFAnimGraph/Private/GraphInterfaces/AnimNextNativeDataInterface_AnimSequencePlayer.cpp
+++ b/Engine/Plugins/Experimental/UAF/UAFAnimGraph/Source/UAFAnimGraph/Private/GraphInterfaces/AnimNextNativeDataInterface_AnimSequencePlayer.cpp
@@ -1,8 +1,10 @@
-﻿// Copyright Epic Games, Inc. All Rights Reserved.
+// Copyright Epic Games, Inc. All Rights Reserved.
 
 #include "GraphInterfaces/AnimNextNativeDataInterface_AnimSequencePlayer.h"
 #include "Animation/AnimSequence.h"
 
+#include UE_INLINE_GENERATED_CPP_BY_NAME(AnimNextNativeDataInterface_AnimSequencePlayer)
+
 void FAnimNextNativeDataInterface_AnimSequencePlayer::BindToFactoryObject(const FBindToFactoryObjectContext& InContext)
 {
 	AnimSequence = CastChecked<UAnimSequence>(InContext.FactoryObject);
diff --git a/Engine/Plugins/Experimental/UAF/UAFAnimGraph/Source/UAFAnimGraph/Private/GraphInterfaces/AnimNextNativeDataInterface_BlendSpacePlayer.cpp b/Engine/Plugins/Experimental/UAF/UAFAnimGraph/Source/UAFAnimGraph/Private/GraphInterfaces/AnimNextNativeDataInterface_BlendSpacePlayer.cpp
--- a/Engine/Plugins/Experimental/UAF/UAFAnimGraph/Source/UAFAnimGraph/Private/GraphInterfaces/AnimNextNativeDataInterface_BlendSpacePlayer.cpp
+++ b/Engine/Plugins/Experimental/UAF/UAFAnimGraph/Source/UAFAnimGraph/Private/GraphInterfaces/AnimNextNativeDataInterface_BlendSpacePlayer.cpp
@@ -3,6 +3,8 @@
 #include "GraphInterfaces/AnimNextNativeDataInterface_BlendSpacePlayer.h"
 #include "Animation/BlendSpace.h"
 
+#include UE_INLINE_GENERATED_CPP_BY_NAME(AnimNextNativeDataInterface_BlendSpacePlayer)
+
 void FAnimNextNativeDataInterface_BlendSpacePlayer::BindToFactoryObject(const FBindToFactoryObjectContext& InContext)
 {
 	BlendSpace = CastChecked<UBlendSpace>(InContext.FactoryObject);
diff --git a/Engine/Plugins/Experimental/UAF/UAFAnimGraph/Source/UAFAnimGraph/Private/Injection/AnimNextModuleInjectionComponent.cpp b/Engine/Plugins/Experimental/UAF/UAFAnimGraph/Source/UAFAnimGraph/Private/Injection/AnimNextModuleInjectionComponent.cpp
--- a/Engine/Plugins/Experimental/UAF/UAFAnimGraph/Source/UAFAnimGraph/Private/Injection/AnimNextModuleInjectionComponent.cpp
+++ b/Engine/Plugins/Experimental/UAF/UAFAnimGraph/Source/UAFAnimGraph/Private/Injection/AnimNextModuleInjectionComponent.cpp
@@ -9,6 +9,8 @@
 #include "Module/AnimNextModuleInstance.h"
 #include "Module/ModuleTaskContext.h"
 
+#include UE_INLINE_GENERATED_CPP_BY_NAME(AnimNextModuleInjectionComponent)
+
 FAnimNextModuleInjectionComponent::FAnimNextModuleInjectionComponent()
 {
 	FAnimNextModuleInstance& ModuleInstance = GetModuleInstance();
diff --git a/Engine/Plugins/Experimental/UAF/UAFAnimGraph/Source/UAFAnimGraph/Private/Injection/GraphInstanceInjectionComponent.cpp b/Engine/Plugins/Experimental/UAF/UAFAnimGraph/Source/UAFAnimGraph/Private/Injection/GraphInstanceInjectionComponent.cpp
--- a/Engine/Plugins/Experimental/UAF/UAFAnimGraph/Source/UAFAnimGraph/Private/Injection/GraphInstanceInjectionComponent.cpp
+++ b/Engine/Plugins/Experimental/UAF/UAFAnimGraph/Source/UAFAnimGraph/Private/Injection/GraphInstanceInjectionComponent.cpp
@@ -1,8 +1,10 @@
-﻿// Copyright Epic Games, Inc. All Rights Reserved.
+// Copyright Epic Games, Inc. All Rights Reserved.
 
 #include "GraphInstanceInjectionComponent.h"
 #include "Graph/AnimNextGraphInstance.h"
 
+#include UE_INLINE_GENERATED_CPP_BY_NAME(GraphInstanceInjectionComponent)
+
 FGraphInstanceInjectionComponent::FGraphInstanceInjectionComponent()
 	: FUAFGraphInstanceComponent()
 	, InjectionInfo(GetGraphInstance())
diff --git a/Engine/Plugins/Experimental/UAF/UAFAnimGraph/Source/UAFAnimGraph/Private/Injection/InjectionRequest.cpp b/Engine/Plugins/Experimental/UAF/UAFAnimGraph/Source/UAFAnimGraph/Private/Injection/InjectionRequest.cpp
--- a/Engine/Plugins/Experimental/UAF/UAFAnimGraph/Source/UAFAnimGraph/Private/Injection/InjectionRequest.cpp
+++ b/Engine/Plugins/Experimental/UAF/UAFAnimGraph/Source/UAFAnimGraph/Private/Injection/InjectionRequest.cpp
@@ -9,6 +9,8 @@
 #include "Logging/StructuredLog.h"
 #include "Engine/World.h"
 
+#include UE_INLINE_GENERATED_CPP_BY_NAME(InjectionRequest)
+
 namespace UE::UAF
 {
 	bool FInjectionRequest::Inject(FInjectionRequestArgs&& InRequestArgs, FInjectionLifetimeEvents&& InLifetimeEvents, UObject* InHost, FModuleHandle InHandle)
diff --git a/Engine/Plugins/Experimental/UAF/UAFAnimGraph/Source/UAFAnimGraph/Private/Injection/InjectionSite.cpp b/Engine/Plugins/Experimental/UAF/UAFAnimGraph/Source/UAFAnimGraph/Private/Injection/InjectionSite.cpp
--- a/Engine/Plugins/Experimental/UAF/UAFAnimGraph/Source/UAFAnimGraph/Private/Injection/InjectionSite.cpp
+++ b/Engine/Plugins/Experimental/UAF/UAFAnimGraph/Source/UAFAnimGraph/Private/Injection/InjectionSite.cpp
@@ -1,8 +1,10 @@
-﻿// Copyright Epic Games, Inc. All Rights Reserved.
+// Copyright Epic Games, Inc. All Rights Reserved.
 
 #include "Injection/InjectionSite.h"
 #include "UObject/FortniteMainBranchObjectVersion.h"
 
+#include UE_INLINE_GENERATED_CPP_BY_NAME(InjectionSite)
+
 #if WITH_EDITORONLY_DATA
 bool FAnimNextInjectionSite::Serialize(FArchive& Ar)
 {
diff --git a/Engine/Plugins/Experimental/UAF/UAFAnimGraph/Source/UAFAnimGraph/Private/Injection/InjectionSiteTrait.cpp b/Engine/Plugins/Experimental/UAF/UAFAnimGraph/Source/UAFAnimGraph/Private/Injection/InjectionSiteTrait.cpp
--- a/Engine/Plugins/Experimental/UAF/UAFAnimGraph/Source/UAFAnimGraph/Private/Injection/InjectionSiteTrait.cpp
+++ b/Engine/Plugins/Experimental/UAF/UAFAnimGraph/Source/UAFAnimGraph/Private/Injection/InjectionSiteTrait.cpp
@@ -18,6 +18,8 @@
 #include "TraitInterfaces/IInertializerBlend.h"
 #include "TraitInterfaces/ISmoothBlend.h"
 
+#include UE_INLINE_GENERATED_CPP_BY_NAME(InjectionSiteTrait)
+
 namespace UE::UAF
 {
 	AUTO_REGISTER_ANIM_TRAIT(FInjectionSiteTrait)
diff --git a/Engine/Plugins/Experimental/UAF/UAFAnimGraph/Source/UAFAnimGraph/Private/PlayAnim/PlayAnimSlotTrait.cpp b/Engine/Plugins/Experimental/UAF/UAFAnimGraph/Source/UAFAnimGraph/Private/PlayAnim/PlayAnimSlotTrait.cpp
--- a/Engine/Plugins/Experimental/UAF/UAFAnimGraph/Source/UAFAnimGraph/Private/PlayAnim/PlayAnimSlotTrait.cpp
+++ b/Engine/Plugins/Experimental/UAF/UAFAnimGraph/Source/UAFAnimGraph/Private/PlayAnim/PlayAnimSlotTrait.cpp
@@ -13,6 +13,8 @@
 #include "Module/ModuleEvents.h"
 #include "TraitInterfaces/IGraphFactory.h"
 
+#include UE_INLINE_GENERATED_CPP_BY_NAME(PlayAnimSlotTrait)
+
 namespace UE::UAF
 {
 	AUTO_REGISTER_ANIM_TRAIT(FPlayAnimSlotTrait)
diff --git a/Engine/Plugins/Experimental/UAF/UAFAnimGraph/Source/UAFAnimGraph/Private/TraitCore/TraitHandle.cpp b/Engine/Plugins/Experimental/UAF/UAFAnimGraph/Source/UAFAnimGraph/Private/TraitCore/TraitHandle.cpp
--- a/Engine/Plugins/Experimental/UAF/UAFAnimGraph/Source/UAFAnimGraph/Private/TraitCore/TraitHandle.cpp
+++ b/Engine/Plugins/Experimental/UAF/UAFAnimGraph/Source/UAFAnimGraph/Private/TraitCore/TraitHandle.cpp
@@ -4,6 +4,8 @@
 #include "TraitCore/TraitReader.h"
 #include "Serialization/Archive.h"
 
+#include UE_INLINE_GENERATED_CPP_BY_NAME(TraitHandle)
+
 bool FAnimNextTraitHandle::Serialize(FArchive& Ar)
 {
 	Ar << PackedTraitIndexAndNodeHandle;
diff --git a/Engine/Plugins/Experimental/UAF/UAFAnimGraph/Source/UAFAnimGraph/Private/TraitInterfaces/Args/IAlphaInputArgs.cpp b/Engine/Plugins/Experimental/UAF/UAFAnimGraph/Source/UAFAnimGraph/Private/TraitInterfaces/Args/IAlphaInputArgs.cpp
--- a/Engine/Plugins/Experimental/UAF/UAFAnimGraph/Source/UAFAnimGraph/Private/TraitInterfaces/Args/IAlphaInputArgs.cpp
+++ b/Engine/Plugins/Experimental/UAF/UAFAnimGraph/Source/UAFAnimGraph/Private/TraitInterfaces/Args/IAlphaInputArgs.cpp
@@ -4,6 +4,8 @@
 
 #include "TraitCore/ExecutionContext.h"
 
+#include UE_INLINE_GENERATED_CPP_BY_NAME(IAlphaInputArgs)
+
 namespace UE::UAF
 {
 	AUTO_REGISTER_ANIM_TRAIT_INTERFACE(IAlphaInputArgs)
diff --git a/Engine/Plugins/Experimental/UAF/UAFAnimGraph/Source/UAFAnimGraph/Private/TraitInterfaces/IGroupSynchronization.cpp b/Engine/Plugins/Experimental/UAF/UAFAnimGraph/Source/UAFAnimGraph/Private/TraitInterfaces/IGroupSynchronization.cpp
--- a/Engine/Plugins/Experimental/UAF/UAFAnimGraph/Source/UAFAnimGraph/Private/TraitInterfaces/IGroupSynchronization.cpp
+++ b/Engine/Plugins/Experimental/UAF/UAFAnimGraph/Source/UAFAnimGraph/Private/TraitInterfaces/IGroupSynchronization.cpp
@@ -4,6 +4,8 @@
 
 #include "TraitCore/ExecutionContext.h"
 
+#include UE_INLINE_GENERATED_CPP_BY_NAME(IGroupSynchronization)
+
 namespace UE::UAF
 {
 	AUTO_REGISTER_ANIM_TRAIT_INTERFACE(IGroupSynchronization)
diff --git a/Engine/Plugins/Experimental/UAF/UAFAnimGraph/Source/UAFAnimGraph/Private/Traits/ApplyAdditiveTrait.cpp b/Engine/Plugins/Experimental/UAF/UAFAnimGraph/Source/UAFAnimGraph/Private/Traits/ApplyAdditiveTrait.cpp
--- a/Engine/Plugins/Experimental/UAF/UAFAnimGraph/Source/UAFAnimGraph/Private/Traits/ApplyAdditiveTrait.cpp
+++ b/Engine/Plugins/Experimental/UAF/UAFAnimGraph/Source/UAFAnimGraph/Private/Traits/ApplyAdditiveTrait.cpp
@@ -7,6 +7,8 @@
 #include "TraitCore/ExecutionContext.h"
 #include "EvaluationVM/Tasks/ApplyAdditiveKeyframe.h"
 
+#include UE_INLINE_GENERATED_CPP_BY_NAME(ApplyAdditiveTrait)
+
 namespace UE::UAF
 {
 	AUTO_REGISTER_ANIM_TRAIT(FApplyAdditiveTrait)
diff --git a/Engine/Plugins/Experimental/UAF/UAFAnimGraph/Source/UAFAnimGraph/Private/Traits/Args/BoolAlphaInputArgTrait.cpp b/Engine/Plugins/Experimental/UAF/UAFAnimGraph/Source/UAFAnimGraph/Private/Traits/Args/BoolAlphaInputArgTrait.cpp
--- a/Engine/Plugins/Experimental/UAF/UAFAnimGraph/Source/UAFAnimGraph/Private/Traits/Args/BoolAlphaInputArgTrait.cpp
+++ b/Engine/Plugins/Experimental/UAF/UAFAnimGraph/Source/UAFAnimGraph/Private/Traits/Args/BoolAlphaInputArgTrait.cpp
@@ -8,6 +8,8 @@
 #include "EvaluationVM/Tasks/BlendKeyframes.h"
 #include "EvaluationVM/Tasks/StoreKeyframe.h"
 
+#include UE_INLINE_GENERATED_CPP_BY_NAME(BoolAlphaInputArgTrait)
+
 namespace UE::UAF
 {
 
diff --git a/Engine/Plugins/Experimental/UAF/UAFAnimGraph/Source/UAFAnimGraph/Private/Traits/Args/CurveAlphaInputArgTrait.cpp b/Engine/Plugins/Experimental/UAF/UAFAnimGraph/Source/UAFAnimGraph/Private/Traits/Args/CurveAlphaInputArgTrait.cpp
--- a/Engine/Plugins/Experimental/UAF/UAFAnimGraph/Source/UAFAnimGraph/Private/Traits/Args/CurveAlphaInputArgTrait.cpp
+++ b/Engine/Plugins/Experimental/UAF/UAFAnimGraph/Source/UAFAnimGraph/Private/Traits/Args/CurveAlphaInputArgTrait.cpp
@@ -8,6 +8,8 @@
 #include "EvaluationVM/Tasks/BlendKeyframes.h"
 #include "EvaluationVM/Tasks/StoreKeyframe.h"
 
+#include UE_INLINE_GENERATED_CPP_BY_NAME(CurveAlphaInputArgTrait)
+
 namespace UE::UAF
 {
 
diff --git a/Engine/Plugins/Experimental/UAF/UAFAnimGraph/Source/UAFAnimGraph/Private/Traits/Args/FloatAlphaInputArgTrait.cpp b/Engine/Plugins/Experimental/UAF/UAFAnimGraph/Source/UAFAnimGraph/Private/Traits/Args/FloatAlphaInputArgTrait.cpp
--- a/Engine/Plugins/Experimental/UAF/UAFAnimGraph/Source/UAFAnimGraph/Private/Traits/Args/FloatAlphaInputArgTrait.cpp
+++ b/Engine/Plugins/Experimental/UAF/UAFAnimGraph/Source/UAFAnimGraph/Private/Traits/Args/FloatAlphaInputArgTrait.cpp
@@ -8,6 +8,8 @@
 #include "EvaluationVM/Tasks/BlendKeyframes.h"
 #include "EvaluationVM/Tasks/StoreKeyframe.h"
 
+#include UE_INLINE_GENERATED_CPP_BY_NAME(FloatAlphaInputArgTrait)
+
 namespace UE::UAF
 {
 
diff --git a/Engine/Plugins/Experimental/UAF/UAFAnimGraph/Source/UAFAnimGraph/Private/Traits/Args/SimpleFloatAlphaInputArgTrait.cpp b/Engine/Plugins/Experimental/UAF/UAFAnimGraph/Source/UAFAnimGraph/Private/Traits/Args/SimpleFloatAlphaInputArgTrait.cpp
--- a/Engine/Plugins/Experimental/UAF/UAFAnimGraph/Source/UAFAnimGraph/Private/Traits/Args/SimpleFloatAlphaInputArgTrait.cpp
+++ b/Engine/Plugins/Experimental/UAF/UAFAnimGraph/Source/UAFAnimGraph/Private/Traits/Args/SimpleFloatAlphaInputArgTrait.cpp
@@ -8,6 +8,8 @@
 #include "EvaluationVM/Tasks/BlendKeyframes.h"
 #include "EvaluationVM/Tasks/StoreKeyframe.h"
 
+#include UE_INLINE_GENERATED_CPP_BY_NAME(SimpleFloatAlphaInputArgTrait)
+
 namespace UE::UAF
 {
 
diff --git a/Engine/Plugins/Experimental/UAF/UAFAnimGraph/Source/UAFAnimGraph/Private/Traits/BlendByBool.cpp b/Engine/Plugins/Experimental/UAF/UAFAnimGraph/Source/UAFAnimGraph/Private/Traits/BlendByBool.cpp
--- a/Engine/Plugins/Experimental/UAF/UAFAnimGraph/Source/UAFAnimGraph/Private/Traits/BlendByBool.cpp
+++ b/Engine/Plugins/Experimental/UAF/UAFAnimGraph/Source/UAFAnimGraph/Private/Traits/BlendByBool.cpp
@@ -6,6 +6,8 @@
 #include "TraitCore/ExecutionContext.h"
 #include "EvaluationVM/Tasks/BlendKeyframes.h"
 
+#include UE_INLINE_GENERATED_CPP_BY_NAME(BlendByBool)
+
 namespace UE::UAF
 {
 	AUTO_REGISTER_ANIM_TRAIT(FBlendByBoolTrait)
diff --git a/Engine/Plugins/Experimental/UAF/UAFAnimGraph/Source/UAFAnimGraph/Private/Traits/BlendInertializer.cpp b/Engine/Plugins/Experimental/UAF/UAFAnimGraph/Source/UAFAnimGraph/Private/Traits/BlendInertializer.cpp
--- a/Engine/Plugins/Experimental/UAF/UAFAnimGraph/Source/UAFAnimGraph/Private/Traits/BlendInertializer.cpp
+++ b/Engine/Plugins/Experimental/UAF/UAFAnimGraph/Source/UAFAnimGraph/Private/Traits/BlendInertializer.cpp
@@ -6,6 +6,8 @@
 
 #include "Traits/Inertialization.h"
 
+#include UE_INLINE_GENERATED_CPP_BY_NAME(BlendInertializer)
+
 
 namespace UE::UAF
 {
diff --git a/Engine/Plugins/Experimental/UAF/UAFAnimGraph/Source/UAFAnimGraph/Private/Traits/BlendLayer.cpp b/Engine/Plugins/Experimental/UAF/UAFAnimGraph/Source/UAFAnimGraph/Private/Traits/BlendLayer.cpp
--- a/Engine/Plugins/Experimental/UAF/UAFAnimGraph/Source/UAFAnimGraph/Private/Traits/BlendLayer.cpp
+++ b/Engine/Plugins/Experimental/UAF/UAFAnimGraph/Source/UAFAnimGraph/Private/Traits/BlendLayer.cpp
@@ -13,6 +13,8 @@
 #include "MaskProfile/HierarchyTableTypeMask.h"
 #include "SkeletonHierarchyTableType.h"
 
+#include UE_INLINE_GENERATED_CPP_BY_NAME(BlendLayer)
+
 class FHierarchyTableBlendProfileProvider : public IInterpolationIndexProvider
 {
 public:
@@ -258,4 +260,4 @@ namespace UE::UAF
 			return -1.0f;
 		}
 	}
-}
\ No newline at end of file
+}
diff --git a/Engine/Plugins/Experimental/UAF/UAFAnimGraph/Source/UAFAnimGraph/Private/Traits/BlendSmoother.cpp b/Engine/Plugins/Experimental/UAF/UAFAnimGraph/Source/UAFAnimGraph/Private/Traits/BlendSmoother.cpp
--- a/Engine/Plugins/Experimental/UAF/UAFAnimGraph/Source/UAFAnimGraph/Private/Traits/BlendSmoother.cpp
+++ b/Engine/Plugins/Experimental/UAF/UAFAnimGraph/Source/UAFAnimGraph/Private/Traits/BlendSmoother.cpp
@@ -8,6 +8,8 @@
 #include "EvaluationVM/Tasks/BlendKeyframes.h"
 #include "EvaluationVM/Tasks/NormalizeRotations.h"
 
+#include UE_INLINE_GENERATED_CPP_BY_NAME(BlendSmoother)
+
 namespace UE::UAF
 {
 	AUTO_REGISTER_ANIM_TRAIT(FBlendSmootherCoreTrait)
diff --git a/Engine/Plugins/Experimental/UAF/UAFAnimGraph/Source/UAFAnimGraph/Private/Traits/BlendSmootherPerBone.cpp b/Engine/Plugins/Experimental/UAF/UAFAnimGraph/Source/UAFAnimGraph/Private/Traits/BlendSmootherPerBone.cpp
--- a/Engine/Plugins/Experimental/UAF/UAFAnimGraph/Source/UAFAnimGraph/Private/Traits/BlendSmootherPerBone.cpp
+++ b/Engine/Plugins/Experimental/UAF/UAFAnimGraph/Source/UAFAnimGraph/Private/Traits/BlendSmootherPerBone.cpp
@@ -10,6 +10,8 @@
 #include "EvaluationVM/Tasks/BlendKeyframesPerBone.h"
 #include "EvaluationVM/Tasks/NormalizeRotations.h"
 
+#include UE_INLINE_GENERATED_CPP_BY_NAME(BlendSmootherPerBone)
+
 namespace UE::UAF
 {
 	AUTO_REGISTER_ANIM_TRAIT(FBlendSmootherPerBoneCoreTrait)
diff --git a/Engine/Plugins/Experimental/UAF/UAFAnimGraph/Source/UAFAnimGraph/Private/Traits/BlendStackTrait.cpp b/Engine/Plugins/Experimental/UAF/UAFAnimGraph/Source/UAFAnimGraph/Private/Traits/BlendStackTrait.cpp
--- a/Engine/Plugins/Experimental/UAF/UAFAnimGraph/Source/UAFAnimGraph/Private/Traits/BlendStackTrait.cpp
+++ b/Engine/Plugins/Experimental/UAF/UAFAnimGraph/Source/UAFAnimGraph/Private/Traits/BlendStackTrait.cpp
@@ -15,6 +15,8 @@
 #include "HierarchyTableBlendProfile.h"
 #include "Factory/AnimGraphFactory.h"
 
+#include UE_INLINE_GENERATED_CPP_BY_NAME(BlendStackTrait)
+
 namespace UE::UAF
 {
 	AUTO_REGISTER_ANIM_TRAIT(FBlendStackCoreTrait)
diff --git a/Engine/Plugins/Experimental/UAF/UAFAnimGraph/Source/UAFAnimGraph/Private/Traits/BlendTwoWay.cpp b/Engine/Plugins/Experimental/UAF/UAFAnimGraph/Source/UAFAnimGraph/Private/Traits/BlendTwoWay.cpp
--- a/Engine/Plugins/Experimental/UAF/UAFAnimGraph/Source/UAFAnimGraph/Private/Traits/BlendTwoWay.cpp
+++ b/Engine/Plugins/Experimental/UAF/UAFAnimGraph/Source/UAFAnimGraph/Private/Traits/BlendTwoWay.cpp
@@ -6,6 +6,8 @@
 #include "TraitCore/ExecutionContext.h"
 #include "EvaluationVM/Tasks/BlendKeyframes.h"
 
+#include UE_INLINE_GENERATED_CPP_BY_NAME(BlendTwoWay)
+
 namespace UE::UAF
 {
 	AUTO_REGISTER_ANIM_TRAIT(FBlendTwoWayTrait)
diff --git a/Engine/Plugins/Experimental/UAF/UAFAnimGraph/Source/UAFAnimGraph/Private/Traits/CallFunction.cpp b/Engine/Plugins/Experimental/UAF/UAFAnimGraph/Source/UAFAnimGraph/Private/Traits/CallFunction.cpp
--- a/Engine/Plugins/Experimental/UAF/UAFAnimGraph/Source/UAFAnimGraph/Private/Traits/CallFunction.cpp
+++ b/Engine/Plugins/Experimental/UAF/UAFAnimGraph/Source/UAFAnimGraph/Private/Traits/CallFunction.cpp
@@ -1,4 +1,4 @@
-﻿// Copyright Epic Games, Inc. All Rights Reserved.
+// Copyright Epic Games, Inc. All Rights Reserved.
 
 #include "Traits/CallFunction.h"
 
@@ -11,6 +11,8 @@
 #include "AnimNextExecuteContext.h"
 #include "UAFRigVMComponent.h"
 
+#include UE_INLINE_GENERATED_CPP_BY_NAME(CallFunction)
+
 DEFINE_STAT(STAT_AnimNext_CallFunction);
 
 namespace UE::UAF
diff --git a/Engine/Plugins/Experimental/UAF/UAFAnimGraph/Source/UAFAnimGraph/Private/Traits/InertializationDeadBlend.cpp b/Engine/Plugins/Experimental/UAF/UAFAnimGraph/Source/UAFAnimGraph/Private/Traits/InertializationDeadBlend.cpp
--- a/Engine/Plugins/Experimental/UAF/UAFAnimGraph/Source/UAFAnimGraph/Private/Traits/InertializationDeadBlend.cpp
+++ b/Engine/Plugins/Experimental/UAF/UAFAnimGraph/Source/UAFAnimGraph/Private/Traits/InertializationDeadBlend.cpp
@@ -10,6 +10,8 @@
 
 #include "Curves/CurveFloat.h"
 
+#include UE_INLINE_GENERATED_CPP_BY_NAME(InertializationDeadBlend)
+
 namespace UE::UAF
 {
 	AUTO_REGISTER_ANIM_TRAIT(FInertializationDeadBlendTrait)
diff --git a/Engine/Plugins/Experimental/UAF/UAFAnimGraph/Source/UAFAnimGraph/Private/Traits/InputPoseTrait.cpp b/Engine/Plugins/Experimental/UAF/UAFAnimGraph/Source/UAFAnimGraph/Private/Traits/InputPoseTrait.cpp
--- a/Engine/Plugins/Experimental/UAF/UAFAnimGraph/Source/UAFAnimGraph/Private/Traits/InputPoseTrait.cpp
+++ b/Engine/Plugins/Experimental/UAF/UAFAnimGraph/Source/UAFAnimGraph/Private/Traits/InputPoseTrait.cpp
@@ -4,6 +4,8 @@
 #include "TraitCore/ExecutionContext.h"
 #include "EvaluationVM/Tasks/PushPose.h"
 
+#include UE_INLINE_GENERATED_CPP_BY_NAME(InputPoseTrait)
+
 namespace UE::UAF
 {
 	AUTO_REGISTER_ANIM_TRAIT(FInputPoseTrait)
@@ -25,4 +27,4 @@ namespace UE::UAF
 			Context.AppendTask(FAnimNextPushPoseTask::Make(&InputPose));
 		}
 	}
-}
\ No newline at end of file
+}
diff --git a/Engine/Plugins/Experimental/UAF/UAFAnimGraph/Source/UAFAnimGraph/Private/Traits/ModifyCurveTrait.cpp b/Engine/Plugins/Experimental/UAF/UAFAnimGraph/Source/UAFAnimGraph/Private/Traits/ModifyCurveTrait.cpp
--- a/Engine/Plugins/Experimental/UAF/UAFAnimGraph/Source/UAFAnimGraph/Private/Traits/ModifyCurveTrait.cpp
+++ b/Engine/Plugins/Experimental/UAF/UAFAnimGraph/Source/UAFAnimGraph/Private/Traits/ModifyCurveTrait.cpp
@@ -4,6 +4,8 @@
 
 #include "EvaluationVM/EvaluationVM.h"
 
+#include UE_INLINE_GENERATED_CPP_BY_NAME(ModifyCurveTrait)
+
 namespace UE::UAF
 {
 AUTO_REGISTER_ANIM_TRAIT(FModifyCurveTrait)
diff --git a/Engine/Plugins/Experimental/UAF/UAFAnimGraph/Source/UAFAnimGraph/Private/Traits/NotifyDispatcher.cpp b/Engine/Plugins/Experimental/UAF/UAFAnimGraph/Source/UAFAnimGraph/Private/Traits/NotifyDispatcher.cpp
--- a/Engine/Plugins/Experimental/UAF/UAFAnimGraph/Source/UAFAnimGraph/Private/Traits/NotifyDispatcher.cpp
+++ b/Engine/Plugins/Experimental/UAF/UAFAnimGraph/Source/UAFAnimGraph/Private/Traits/NotifyDispatcher.cpp
@@ -13,6 +13,8 @@
 #include "Module/AnimNextModuleInstance.h"
 #include "Module/AnimNextSkeletalMeshComponentReferenceComponent.h"
 
+#include UE_INLINE_GENERATED_CPP_BY_NAME(NotifyDispatcher)
+
 namespace UE::UAF
 {
 	AUTO_REGISTER_ANIM_TRAIT(FNotifyDispatcherTrait)
diff --git a/Engine/Plugins/Experimental/UAF/UAFAnimGraph/Source/UAFAnimGraph/Private/Traits/NotifyFilterTrait.cpp b/Engine/Plugins/Experimental/UAF/UAFAnimGraph/Source/UAFAnimGraph/Private/Traits/NotifyFilterTrait.cpp
--- a/Engine/Plugins/Experimental/UAF/UAFAnimGraph/Source/UAFAnimGraph/Private/Traits/NotifyFilterTrait.cpp
+++ b/Engine/Plugins/Experimental/UAF/UAFAnimGraph/Source/UAFAnimGraph/Private/Traits/NotifyFilterTrait.cpp
@@ -3,6 +3,8 @@
 #include "NotifyFilterTrait.h"
 #include "TraitCore/ExecutionContext.h"
 
+#include UE_INLINE_GENERATED_CPP_BY_NAME(NotifyFilterTrait)
+
 namespace UE::UAF
 {
 	AUTO_REGISTER_ANIM_TRAIT(FNotifyFilterTrait)
diff --git a/Engine/Plugins/Experimental/UAF/UAFAnimGraph/Source/UAFAnimGraph/Private/Traits/Passthrough.cpp b/Engine/Plugins/Experimental/UAF/UAFAnimGraph/Source/UAFAnimGraph/Private/Traits/Passthrough.cpp
--- a/Engine/Plugins/Experimental/UAF/UAFAnimGraph/Source/UAFAnimGraph/Private/Traits/Passthrough.cpp
+++ b/Engine/Plugins/Experimental/UAF/UAFAnimGraph/Source/UAFAnimGraph/Private/Traits/Passthrough.cpp
@@ -4,6 +4,8 @@
 
 #include "TraitCore/ExecutionContext.h"
 
+#include UE_INLINE_GENERATED_CPP_BY_NAME(Passthrough)
+
 namespace UE::UAF
 {
 	AUTO_REGISTER_ANIM_TRAIT(FPassthroughTrait)
diff --git a/Engine/Plugins/Experimental/UAF/UAFAnimGraph/Source/UAFAnimGraph/Private/Traits/PassthroughBlendTrait.cpp b/Engine/Plugins/Experimental/UAF/UAFAnimGraph/Source/UAFAnimGraph/Private/Traits/PassthroughBlendTrait.cpp
--- a/Engine/Plugins/Experimental/UAF/UAFAnimGraph/Source/UAFAnimGraph/Private/Traits/PassthroughBlendTrait.cpp
+++ b/Engine/Plugins/Experimental/UAF/UAFAnimGraph/Source/UAFAnimGraph/Private/Traits/PassthroughBlendTrait.cpp
@@ -7,6 +7,8 @@
 #include "EvaluationVM/Tasks/BlendKeyframes.h"
 #include "EvaluationVM/Tasks/StoreKeyframe.h"
 
+#include UE_INLINE_GENERATED_CPP_BY_NAME(PassthroughBlendTrait)
+
 namespace UE::UAF
 {
 
diff --git a/Engine/Plugins/Experimental/UAF/UAFAnimGraph/Source/UAFAnimGraph/Private/Traits/ReferencePoseTrait.cpp b/Engine/Plugins/Experimental/UAF/UAFAnimGraph/Source/UAFAnimGraph/Private/Traits/ReferencePoseTrait.cpp
--- a/Engine/Plugins/Experimental/UAF/UAFAnimGraph/Source/UAFAnimGraph/Private/Traits/ReferencePoseTrait.cpp
+++ b/Engine/Plugins/Experimental/UAF/UAFAnimGraph/Source/UAFAnimGraph/Private/Traits/ReferencePoseTrait.cpp
@@ -5,6 +5,8 @@
 #include "TraitCore/ExecutionContext.h"
 #include "EvaluationVM/Tasks/PushReferenceKeyframe.h"
 
+#include UE_INLINE_GENERATED_CPP_BY_NAME(ReferencePoseTrait)
+
 namespace UE::UAF
 {
 	AUTO_REGISTER_ANIM_TRAIT(FReferencePoseTrait)
diff --git a/Engine/Plugins/Experimental/UAF/UAFAnimGraph/Source/UAFAnimGraph/Private/Traits/SubGraphHost.cpp b/Engine/Plugins/Experimental/UAF/UAFAnimGraph/Source/UAFAnimGraph/Private/Traits/SubGraphHost.cpp
--- a/Engine/Plugins/Experimental/UAF/UAFAnimGraph/Source/UAFAnimGraph/Private/Traits/SubGraphHost.cpp
+++ b/Engine/Plugins/Experimental/UAF/UAFAnimGraph/Source/UAFAnimGraph/Private/Traits/SubGraphHost.cpp
@@ -6,6 +6,8 @@
 #include "TraitCore/NodeInstance.h"
 #include "Graph/AnimNextGraphInstance.h"
 
+#include UE_INLINE_GENERATED_CPP_BY_NAME(SubGraphHost)
+
 namespace UE::UAF
 {
 	AUTO_REGISTER_ANIM_TRAIT(FSubGraphHostTrait)
diff --git a/Engine/Plugins/Experimental/UAF/UAFAnimGraph/Source/UAFAnimGraphEditor/Private/Graph/AnimNextAnimGraphEntryAssetDefinitions.cpp b/Engine/Plugins/Experimental/UAF/UAFAnimGraph/Source/UAFAnimGraphEditor/Private/Graph/AnimNextAnimGraphEntryAssetDefinitions.cpp
--- a/Engine/Plugins/Experimental/UAF/UAFAnimGraph/Source/UAFAnimGraphEditor/Private/Graph/AnimNextAnimGraphEntryAssetDefinitions.cpp
+++ b/Engine/Plugins/Experimental/UAF/UAFAnimGraph/Source/UAFAnimGraphEditor/Private/Graph/AnimNextAnimGraphEntryAssetDefinitions.cpp
@@ -1,7 +1,9 @@
-﻿// Copyright Epic Games, Inc. All Rights Reserved.
+// Copyright Epic Games, Inc. All Rights Reserved.
 
 #include "AnimNextAnimGraphEntryAssetDefinitions.h"
 
+#include UE_INLINE_GENERATED_CPP_BY_NAME(AnimNextAnimGraphEntryAssetDefinitions)
+
 #define LOCTEXT_NAMESPACE "AnimNextAnimGraphEntryAssetDefinitions"
 
 FText UAssetDefinition_AnimNextAnimationGraphEntry::GetObjectDisplayNameText(UObject* Object) const
@@ -10,4 +12,4 @@ FText UAssetDefinition_AnimNextAnimationGraphEntry::GetObjectDisplayNameText(UOb
 	return Parameter->GetDisplayName();
 }
 
-#undef LOCTEXT_NAMESPACE
\ No newline at end of file
+#undef LOCTEXT_NAMESPACE
diff --git a/Engine/Plugins/Experimental/UAF/UAFAnimGraph/Source/UAFAnimGraphEditor/Private/Graph/AnimNextAnimationGraphAssetDefinition.cpp b/Engine/Plugins/Experimental/UAF/UAFAnimGraph/Source/UAFAnimGraphEditor/Private/Graph/AnimNextAnimationGraphAssetDefinition.cpp
--- a/Engine/Plugins/Experimental/UAF/UAFAnimGraph/Source/UAFAnimGraphEditor/Private/Graph/AnimNextAnimationGraphAssetDefinition.cpp
+++ b/Engine/Plugins/Experimental/UAF/UAFAnimGraph/Source/UAFAnimGraphEditor/Private/Graph/AnimNextAnimationGraphAssetDefinition.cpp
@@ -1,4 +1,4 @@
-﻿// Copyright Epic Games, Inc. All Rights Reserved.
+// Copyright Epic Games, Inc. All Rights Reserved.
 
 #include "AnimNextAnimationGraphAssetDefinition.h"
 
@@ -10,6 +10,8 @@
 #include "UncookedOnlyUtils.h"
 #include "Workspace/AnimNextWorkspaceFactory.h"
 
+#include UE_INLINE_GENERATED_CPP_BY_NAME(AnimNextAnimationGraphAssetDefinition)
+
 #define LOCTEXT_NAMESPACE "AnimNextAssetDefinitions"
 
 EAssetCommandResult UAssetDefinition_AnimNextAnimationGraph::OpenAssets(const FAssetOpenArgs& OpenArgs) const
@@ -26,4 +28,4 @@ EAssetCommandResult UAssetDefinition_AnimNextAnimationGraph::OpenAssets(const FA
 	return EAssetCommandResult::Handled;
 }
 
-#undef LOCTEXT_NAMESPACE
\ No newline at end of file
+#undef LOCTEXT_NAMESPACE
diff --git a/Engine/Plugins/Experimental/UAF/UAFAnimGraph/Source/UAFAnimGraphEditor/Private/Graph/AnimNextAnimationGraphFactory.cpp b/Engine/Plugins/Experimental/UAF/UAFAnimGraph/Source/UAFAnimGraphEditor/Private/Graph/AnimNextAnimationGraphFactory.cpp
--- a/Engine/Plugins/Experimental/UAF/UAFAnimGraph/Source/UAFAnimGraphEditor/Private/Graph/AnimNextAnimationGraphFactory.cpp
+++ b/Engine/Plugins/Experimental/UAF/UAFAnimGraph/Source/UAFAnimGraphEditor/Private/Graph/AnimNextAnimationGraphFactory.cpp
@@ -1,4 +1,4 @@
-﻿// Copyright Epic Games, Inc. All Rights Reserved.
+// Copyright Epic Games, Inc. All Rights Reserved.
 
 #include "Graph/AnimNextAnimationGraphFactory.h"
 #include "Graph/AnimNextAnimationGraph.h"
@@ -7,6 +7,8 @@
 #include "Entries/AnimNextAnimationGraphEntry.h"
 #include "Graph/RigUnit_AnimNextGraphRoot.h"
 
+#include UE_INLINE_GENERATED_CPP_BY_NAME(AnimNextAnimationGraphFactory)
+
 UAnimNextAnimationGraphFactory::UAnimNextAnimationGraphFactory()
 {
 	bCreateNew = true;
@@ -46,4 +48,4 @@ UObject* UAnimNextAnimationGraphFactory::FactoryCreateNew(UClass* Class, UObject
 	check(!EditorData->bErrorsDuringCompilation);
 
 	return NewAnimationGraph;
-}
\ No newline at end of file
+}
diff --git a/Engine/Plugins/Experimental/UAF/UAFAnimGraph/Source/UAFAnimGraphEditor/Private/RewindDebugger/EvaluationProgramTrack.cpp b/Engine/Plugins/Experimental/UAF/UAFAnimGraph/Source/UAFAnimGraphEditor/Private/RewindDebugger/EvaluationProgramTrack.cpp
--- a/Engine/Plugins/Experimental/UAF/UAFAnimGraph/Source/UAFAnimGraphEditor/Private/RewindDebugger/EvaluationProgramTrack.cpp
+++ b/Engine/Plugins/Experimental/UAF/UAFAnimGraph/Source/UAFAnimGraphEditor/Private/RewindDebugger/EvaluationProgramTrack.cpp
@@ -12,6 +12,8 @@
 #include "ObjectAsTraceIdProxyArchiveReader.h"
 #include "Serialization/ObjectReader.h"
 
+#include UE_INLINE_GENERATED_CPP_BY_NAME(EvaluationProgramTrack)
+
 
 #define LOCTEXT_NAMESPACE "EvaluationProgramTrack"
 
@@ -200,4 +202,4 @@ bool FEvaluationProgramTrackCreator::HasDebugInfoInternal(const RewindDebugger::
 
 }
 
-#undef LOCTEXT_NAMESPACE
\ No newline at end of file
+#undef LOCTEXT_NAMESPACE
diff --git a/Engine/Plugins/Experimental/UAF/UAFAnimGraph/Source/UAFAnimGraphEditor/Private/RewindDebugger/SequenceInfoTrack.cpp b/Engine/Plugins/Experimental/UAF/UAFAnimGraph/Source/UAFAnimGraphEditor/Private/RewindDebugger/SequenceInfoTrack.cpp
--- a/Engine/Plugins/Experimental/UAF/UAFAnimGraph/Source/UAFAnimGraphEditor/Private/RewindDebugger/SequenceInfoTrack.cpp
+++ b/Engine/Plugins/Experimental/UAF/UAFAnimGraph/Source/UAFAnimGraphEditor/Private/RewindDebugger/SequenceInfoTrack.cpp
@@ -16,6 +16,8 @@
 #include "Widgets/Text/STextBlock.h"
 #include "Widgets/Colors/SColorBlock.h"
 
+#include UE_INLINE_GENERATED_CPP_BY_NAME(SequenceInfoTrack)
+
 #define LOCTEXT_NAMESPACE "EvaluationProgramTrack"
 
 namespace UE::UAF::Editor
diff --git a/Engine/Plugins/Experimental/UAF/UAFAnimGraph/Source/UAFAnimGraphTestSuite/Private/AnimNextAnimGraphTraitGraphTest.cpp b/Engine/Plugins/Experimental/UAF/UAFAnimGraph/Source/UAFAnimGraphTestSuite/Private/AnimNextAnimGraphTraitGraphTest.cpp
--- a/Engine/Plugins/Experimental/UAF/UAFAnimGraph/Source/UAFAnimGraphTestSuite/Private/AnimNextAnimGraphTraitGraphTest.cpp
+++ b/Engine/Plugins/Experimental/UAF/UAFAnimGraph/Source/UAFAnimGraphTestSuite/Private/AnimNextAnimGraphTraitGraphTest.cpp
@@ -27,6 +27,8 @@
 #include "AnimGraphUncookedOnlyUtils.h"
 #include "UAFRigVMComponent.h"
 
+#include UE_INLINE_GENERATED_CPP_BY_NAME(AnimNextAnimGraphTraitGraphTest)
+
 #if WITH_DEV_AUTOMATION_TESTS
 
 //****************************************************************************
diff --git a/Engine/Plugins/Experimental/UAF/UAFAnimGraph/Source/UAFAnimGraphTestSuite/Private/AnimNextAnimGraphTraitInterfacesTest.cpp b/Engine/Plugins/Experimental/UAF/UAFAnimGraph/Source/UAFAnimGraphTestSuite/Private/AnimNextAnimGraphTraitInterfacesTest.cpp
--- a/Engine/Plugins/Experimental/UAF/UAFAnimGraph/Source/UAFAnimGraphTestSuite/Private/AnimNextAnimGraphTraitInterfacesTest.cpp
+++ b/Engine/Plugins/Experimental/UAF/UAFAnimGraph/Source/UAFAnimGraphTestSuite/Private/AnimNextAnimGraphTraitInterfacesTest.cpp
@@ -7,6 +7,8 @@
 #include "Misc/AutomationTest.h"
 #include "TraitInterfaces/IScopedTag.h"
 
+#include UE_INLINE_GENERATED_CPP_BY_NAME(AnimNextAnimGraphTraitInterfacesTest)
+
 #if WITH_DEV_AUTOMATION_TESTS
 
 #include "Serialization/MemoryReader.h"
diff --git a/Engine/Plugins/Experimental/UAF/UAFAnimGraph/Source/UAFAnimGraphTestSuite/Private/AnimNextTraitBaseTest.cpp b/Engine/Plugins/Experimental/UAF/UAFAnimGraph/Source/UAFAnimGraphTestSuite/Private/AnimNextTraitBaseTest.cpp
--- a/Engine/Plugins/Experimental/UAF/UAFAnimGraph/Source/UAFAnimGraphTestSuite/Private/AnimNextTraitBaseTest.cpp
+++ b/Engine/Plugins/Experimental/UAF/UAFAnimGraph/Source/UAFAnimGraphTestSuite/Private/AnimNextTraitBaseTest.cpp
@@ -9,6 +9,8 @@
 
 #include "Misc/AutomationTest.h"
 
+#include UE_INLINE_GENERATED_CPP_BY_NAME(AnimNextTraitBaseTest)
+
 #if WITH_DEV_AUTOMATION_TESTS
 
 #include "Serialization/MemoryReader.h"
diff --git a/Engine/Plugins/Experimental/UAF/UAFAnimGraph/Source/UAFAnimGraphUncookedOnly/Internal/Graph/AnimGraphEditorSchemaActions.cpp b/Engine/Plugins/Experimental/UAF/UAFAnimGraph/Source/UAFAnimGraphUncookedOnly/Internal/Graph/AnimGraphEditorSchemaActions.cpp
--- a/Engine/Plugins/Experimental/UAF/UAFAnimGraph/Source/UAFAnimGraphUncookedOnly/Internal/Graph/AnimGraphEditorSchemaActions.cpp
+++ b/Engine/Plugins/Experimental/UAF/UAFAnimGraph/Source/UAFAnimGraphUncookedOnly/Internal/Graph/AnimGraphEditorSchemaActions.cpp
@@ -21,6 +21,8 @@
 #include "RigVMFunctions/Execution/RigVMFunction_UserDefinedEvent.h"
 #include "Widgets/Input/STextEntryPopup.h"
 
+#include UE_INLINE_GENERATED_CPP_BY_NAME(AnimGraphEditorSchemaActions)
+
 #define LOCTEXT_NAMESPACE "AnimNextAnimGraphSchemaActions"
 
 namespace Private
diff --git a/Engine/Plugins/Experimental/UAF/UAFAnimGraph/Source/UAFAnimGraphUncookedOnly/Private/AnimNextAnimGraphWorkspaceAssetUserData.cpp b/Engine/Plugins/Experimental/UAF/UAFAnimGraph/Source/UAFAnimGraphUncookedOnly/Private/AnimNextAnimGraphWorkspaceAssetUserData.cpp
--- a/Engine/Plugins/Experimental/UAF/UAFAnimGraph/Source/UAFAnimGraphUncookedOnly/Private/AnimNextAnimGraphWorkspaceAssetUserData.cpp
+++ b/Engine/Plugins/Experimental/UAF/UAFAnimGraph/Source/UAFAnimGraphUncookedOnly/Private/AnimNextAnimGraphWorkspaceAssetUserData.cpp
@@ -1,4 +1,4 @@
-﻿// Copyright Epic Games, Inc. All Rights Reserved.
+// Copyright Epic Games, Inc. All Rights Reserved.
 
 #include "AnimNextAnimGraphWorkspaceAssetUserData.h"
 
@@ -6,6 +6,8 @@
 #include "Graph/AnimNextAnimationGraph_EditorData.h"
 #include "UObject/AssetRegistryTagsContext.h"
 
+#include UE_INLINE_GENERATED_CPP_BY_NAME(AnimNextAnimGraphWorkspaceAssetUserData)
+
 void UAnimNextAnimGraphWorkspaceAssetUserData::GetAssetRegistryTags(FAssetRegistryTagsContext Context) const
 {
 	Super::GetAssetRegistryTags(Context);
diff --git a/Engine/Plugins/Experimental/UAF/UAFAnimGraph/Source/UAFAnimGraphUncookedOnly/Private/AnimNextController.cpp b/Engine/Plugins/Experimental/UAF/UAFAnimGraph/Source/UAFAnimGraphUncookedOnly/Private/AnimNextController.cpp
--- a/Engine/Plugins/Experimental/UAF/UAFAnimGraph/Source/UAFAnimGraphUncookedOnly/Private/AnimNextController.cpp
+++ b/Engine/Plugins/Experimental/UAF/UAFAnimGraph/Source/UAFAnimGraphUncookedOnly/Private/AnimNextController.cpp
@@ -13,6 +13,8 @@
 #include "TraitCore/TraitRegistry.h"
 #include "TraitCore/TraitUID.h"
 
+#include UE_INLINE_GENERATED_CPP_BY_NAME(AnimNextController)
+
 FName UAnimNextController::AddTraitByName(FName InNodeName, FName InNewTraitTypeName, int32 InPinIndex, const FString& InNewTraitDefaultValue, bool bSetupUndoRedo, bool bPrintPythonCommand)
 {
 	if (!IsValidGraph())
diff --git a/Engine/Plugins/Experimental/UAF/UAFAnimGraph/Source/UAFAnimGraphUncookedOnly/Private/AnimNextControllerActions.cpp b/Engine/Plugins/Experimental/UAF/UAFAnimGraph/Source/UAFAnimGraphUncookedOnly/Private/AnimNextControllerActions.cpp
--- a/Engine/Plugins/Experimental/UAF/UAFAnimGraph/Source/UAFAnimGraphUncookedOnly/Private/AnimNextControllerActions.cpp
+++ b/Engine/Plugins/Experimental/UAF/UAFAnimGraph/Source/UAFAnimGraphUncookedOnly/Private/AnimNextControllerActions.cpp
@@ -3,6 +3,8 @@
 #include "AnimNextControllerActions.h"
 #include "AnimGraphUncookedOnlyUtils.h"
 
+#include UE_INLINE_GENERATED_CPP_BY_NAME(AnimNextControllerActions)
+
 FAnimNextManifestAction::FAnimNextManifestAction()
 	: FAnimNextBaseAction(nullptr)
 {
diff --git a/Engine/Plugins/Experimental/UAF/UAFAnimGraph/Source/UAFAnimGraphUncookedOnly/Private/Entries/AnimNextAnimationGraphEntry.cpp b/Engine/Plugins/Experimental/UAF/UAFAnimGraph/Source/UAFAnimGraphUncookedOnly/Private/Entries/AnimNextAnimationGraphEntry.cpp
--- a/Engine/Plugins/Experimental/UAF/UAFAnimGraph/Source/UAFAnimGraphUncookedOnly/Private/Entries/AnimNextAnimationGraphEntry.cpp
+++ b/Engine/Plugins/Experimental/UAF/UAFAnimGraph/Source/UAFAnimGraphUncookedOnly/Private/Entries/AnimNextAnimationGraphEntry.cpp
@@ -12,6 +12,8 @@
 #include "Param/AnimNextTag.h"
 #include "Param/ParamType.h"
 
+#include UE_INLINE_GENERATED_CPP_BY_NAME(AnimNextAnimationGraphEntry)
+
 FAnimNextParamType UAnimNextAnimationGraphEntry::GetExportType() const
 {
 	return FAnimNextParamType::GetType<FAnimNextEntryPoint>();
diff --git a/Engine/Plugins/Experimental/UAF/UAFAnimGraph/Source/UAFAnimGraphUncookedOnly/Private/Graph/AnimNextAnimationGraphSchema.cpp b/Engine/Plugins/Experimental/UAF/UAFAnimGraph/Source/UAFAnimGraphUncookedOnly/Private/Graph/AnimNextAnimationGraphSchema.cpp
--- a/Engine/Plugins/Experimental/UAF/UAFAnimGraph/Source/UAFAnimGraphUncookedOnly/Private/Graph/AnimNextAnimationGraphSchema.cpp
+++ b/Engine/Plugins/Experimental/UAF/UAFAnimGraph/Source/UAFAnimGraphUncookedOnly/Private/Graph/AnimNextAnimationGraphSchema.cpp
@@ -3,6 +3,8 @@
 #include "Graph/AnimNextAnimationGraphSchema.h"
 #include "AnimGraphUncookedOnlyUtils.h"
 
+#include UE_INLINE_GENERATED_CPP_BY_NAME(AnimNextAnimationGraphSchema)
+
 bool UAnimNextAnimationGraphSchema::CanSetNodeTitle(URigVMController* InController, const URigVMNode* InNode) const
 {
 	if (InNode != nullptr)
@@ -32,4 +34,4 @@ bool UAnimNextAnimationGraphSchema::CanUnfoldPin(URigVMController* InController,
 		}
 	}
 	return false;
-}
\ No newline at end of file
+}
diff --git a/Engine/Plugins/Experimental/UAF/UAFAnimGraph/Source/UAFAnimGraphUncookedOnly/Private/Graph/AnimNextAnimationGraph_EditorData.cpp b/Engine/Plugins/Experimental/UAF/UAFAnimGraph/Source/UAFAnimGraphUncookedOnly/Private/Graph/AnimNextAnimationGraph_EditorData.cpp
--- a/Engine/Plugins/Experimental/UAF/UAFAnimGraph/Source/UAFAnimGraphUncookedOnly/Private/Graph/AnimNextAnimationGraph_EditorData.cpp
+++ b/Engine/Plugins/Experimental/UAF/UAFAnimGraph/Source/UAFAnimGraphUncookedOnly/Private/Graph/AnimNextAnimationGraph_EditorData.cpp
@@ -31,6 +31,8 @@
 #include "Graph/RigUnit_AnimNextGraphEvaluator.h"
 #include "Graph/RigUnit_AnimNextGraphRoot.h"
 
+#include UE_INLINE_GENERATED_CPP_BY_NAME(AnimNextAnimationGraph_EditorData)
+
 #define LOCTEXT_NAMESPACE "AnimNextAnimationGraph_EditorData"
 
 namespace UE::UAF::UncookedOnly::Private
@@ -871,4 +873,4 @@ void UAnimNextAnimationGraph_EditorData::GetAnimNextAssetRegistryTags(FAssetRegi
 	UAnimNextRigVMAssetEditorData::GetAnimNextAssetRegistryTags(Context, OutExports);
 }
 
-#undef LOCTEXT_NAMESPACE
\ No newline at end of file
+#undef LOCTEXT_NAMESPACE
diff --git a/Engine/Plugins/Experimental/UAF/UAFAnimGraph/Source/UAFAnimGraphUncookedOnly/Private/Injection/K2Node_AnimNextInjection.cpp b/Engine/Plugins/Experimental/UAF/UAFAnimGraph/Source/UAFAnimGraphUncookedOnly/Private/Injection/K2Node_AnimNextInjection.cpp
--- a/Engine/Plugins/Experimental/UAF/UAFAnimGraph/Source/UAFAnimGraphUncookedOnly/Private/Injection/K2Node_AnimNextInjection.cpp
+++ b/Engine/Plugins/Experimental/UAF/UAFAnimGraph/Source/UAFAnimGraphUncookedOnly/Private/Injection/K2Node_AnimNextInjection.cpp
@@ -1,8 +1,10 @@
-﻿// Copyright Epic Games, Inc. All Rights Reserved.
+// Copyright Epic Games, Inc. All Rights Reserved.
 
 #include "Injection/K2Node_AnimNextInjection.h"
 #include "Injection/InjectionCallbackProxy.h"
 
+#include UE_INLINE_GENERATED_CPP_BY_NAME(K2Node_AnimNextInjection)
+
 #define LOCTEXT_NAMESPACE "K2Node_AnimNextInjection"
 
 UK2Node_AnimNextInjection::UK2Node_AnimNextInjection()
diff --git a/Engine/Plugins/Experimental/UAF/UAFAnimGraph/Source/UAFAnimGraphUncookedOnly/Private/PlayAnim/K2Node_AnimNextPlayAnim.cpp b/Engine/Plugins/Experimental/UAF/UAFAnimGraph/Source/UAFAnimGraphUncookedOnly/Private/PlayAnim/K2Node_AnimNextPlayAnim.cpp
--- a/Engine/Plugins/Experimental/UAF/UAFAnimGraph/Source/UAFAnimGraphUncookedOnly/Private/PlayAnim/K2Node_AnimNextPlayAnim.cpp
+++ b/Engine/Plugins/Experimental/UAF/UAFAnimGraph/Source/UAFAnimGraphUncookedOnly/Private/PlayAnim/K2Node_AnimNextPlayAnim.cpp
@@ -4,6 +4,8 @@
 
 #include "PlayAnim/PlayAnimCallbackProxy.h"
 
+#include UE_INLINE_GENERATED_CPP_BY_NAME(K2Node_AnimNextPlayAnim)
+
 #define LOCTEXT_NAMESPACE "K2Node_AnimNextPlayAnim"
 
 UK2Node_AnimNextPlayAnim::UK2Node_AnimNextPlayAnim()
diff --git a/Engine/Plugins/Experimental/UAF/UAFAnimGraph/Source/UAFAnimGraphUncookedOnly/Private/PlayAnim/K2Node_AnimNextPlayAsset.cpp b/Engine/Plugins/Experimental/UAF/UAFAnimGraph/Source/UAFAnimGraphUncookedOnly/Private/PlayAnim/K2Node_AnimNextPlayAsset.cpp
--- a/Engine/Plugins/Experimental/UAF/UAFAnimGraph/Source/UAFAnimGraphUncookedOnly/Private/PlayAnim/K2Node_AnimNextPlayAsset.cpp
+++ b/Engine/Plugins/Experimental/UAF/UAFAnimGraph/Source/UAFAnimGraphUncookedOnly/Private/PlayAnim/K2Node_AnimNextPlayAsset.cpp
@@ -1,9 +1,11 @@
-﻿// Copyright Epic Games, Inc. All Rights Reserved.
+// Copyright Epic Games, Inc. All Rights Reserved.
 
 #include "PlayAnim/K2Node_AnimNextPlayAsset.h"
 
 #include "PlayAnim/PlayAnimCallbackProxy.h"
 
+#include UE_INLINE_GENERATED_CPP_BY_NAME(K2Node_AnimNextPlayAsset)
+
 #define LOCTEXT_NAMESPACE "K2Node_AnimNextPlayAsset"
 
 UDEPRECATED_K2Node_AnimNextPlayAsset::UDEPRECATED_K2Node_AnimNextPlayAsset()
diff --git a/Engine/Plugins/Experimental/Workspace/Source/WorkspaceEditor/Private/Workspace.cpp b/Engine/Plugins/Experimental/Workspace/Source/WorkspaceEditor/Private/Workspace.cpp
--- a/Engine/Plugins/Experimental/Workspace/Source/WorkspaceEditor/Private/Workspace.cpp
+++ b/Engine/Plugins/Experimental/Workspace/Source/WorkspaceEditor/Private/Workspace.cpp
@@ -11,6 +11,8 @@
 #include "UObject/FortniteMainBranchObjectVersion.h"
 #include "Engine/ExternalAssetDependencyGatherer.h"
 
+#include UE_INLINE_GENERATED_CPP_BY_NAME(Workspace)
+
 REGISTER_ASSETDEPENDENCY_GATHERER(FExternalAssetDependencyGatherer, UWorkspace);
 
 bool UWorkspace::AddAsset(const FAssetData& InAsset, bool bSetupUndoRedo, bool bPrintPythonCommand)
diff --git a/Engine/Plugins/Experimental/Workspace/Source/WorkspaceEditor/Private/WorkspaceAssetDefinition.cpp b/Engine/Plugins/Experimental/Workspace/Source/WorkspaceEditor/Private/WorkspaceAssetDefinition.cpp
--- a/Engine/Plugins/Experimental/Workspace/Source/WorkspaceEditor/Private/WorkspaceAssetDefinition.cpp
+++ b/Engine/Plugins/Experimental/Workspace/Source/WorkspaceEditor/Private/WorkspaceAssetDefinition.cpp
@@ -5,6 +5,8 @@
 #include "WorkspaceSchema.h"
 #include "WorkspaceAssetEditor.h"
 
+#include UE_INLINE_GENERATED_CPP_BY_NAME(WorkspaceAssetDefinition)
+
 #define LOCTEXT_NAMESPACE "AssetDefinition_Workspace"
 
 FText UAssetDefinition_Workspace::GetAssetDisplayName() const
@@ -43,4 +45,4 @@ EAssetCommandResult UAssetDefinition_Workspace::OpenAssets(const FAssetOpenArgs&
 	return EAssetCommandResult::Handled;
 }
 
-#undef LOCTEXT_NAMESPACE
\ No newline at end of file
+#undef LOCTEXT_NAMESPACE
diff --git a/Engine/Plugins/Experimental/Workspace/Source/WorkspaceEditor/Private/WorkspaceAssetEntry.cpp b/Engine/Plugins/Experimental/Workspace/Source/WorkspaceEditor/Private/WorkspaceAssetEntry.cpp
--- a/Engine/Plugins/Experimental/Workspace/Source/WorkspaceEditor/Private/WorkspaceAssetEntry.cpp
+++ b/Engine/Plugins/Experimental/Workspace/Source/WorkspaceEditor/Private/WorkspaceAssetEntry.cpp
@@ -5,6 +5,8 @@
 #include "UObject/AssetRegistryTagsContext.h"
 #include "UObject/Package.h"
 
+#include UE_INLINE_GENERATED_CPP_BY_NAME(WorkspaceAssetEntry)
+
 const FName UWorkspaceAssetEntry::ExportsAssetRegistryTag = TEXT("Exports");
 
 bool UWorkspaceAssetEntry::IsAsset() const
diff --git a/Engine/Plugins/Experimental/Workspace/Source/WorkspaceEditor/Private/WorkspaceAssetEntryAssetDefinition.cpp b/Engine/Plugins/Experimental/Workspace/Source/WorkspaceEditor/Private/WorkspaceAssetEntryAssetDefinition.cpp
--- a/Engine/Plugins/Experimental/Workspace/Source/WorkspaceEditor/Private/WorkspaceAssetEntryAssetDefinition.cpp
+++ b/Engine/Plugins/Experimental/Workspace/Source/WorkspaceEditor/Private/WorkspaceAssetEntryAssetDefinition.cpp
@@ -3,6 +3,8 @@
 #include "WorkspaceAssetEntryAssetDefinition.h"
 #include "WorkspaceAssetEntry.h"
 
+#include UE_INLINE_GENERATED_CPP_BY_NAME(WorkspaceAssetEntryAssetDefinition)
+
 #define LOCTEXT_NAMESPACE "AssetDefinition_WorkspaceAssetEntry"
 
 FText UAssetDefinition_WorkspaceAssetEntry::GetAssetDisplayName() const
@@ -26,4 +28,4 @@ FText UAssetDefinition_WorkspaceAssetEntry::GetObjectDisplayNameText(UObject* Ob
 	return FText::FromString(Entry->Asset.GetAssetName());
 }
 
-#undef LOCTEXT_NAMESPACE // "AssetDefinition_WorkspaceAssetEntry"
\ No newline at end of file
+#undef LOCTEXT_NAMESPACE // "AssetDefinition_WorkspaceAssetEntry"
diff --git a/Engine/Plugins/Experimental/Workspace/Source/WorkspaceEditor/Private/WorkspaceAssetRegistryInfo.cpp b/Engine/Plugins/Experimental/Workspace/Source/WorkspaceEditor/Private/WorkspaceAssetRegistryInfo.cpp
--- a/Engine/Plugins/Experimental/Workspace/Source/WorkspaceEditor/Private/WorkspaceAssetRegistryInfo.cpp
+++ b/Engine/Plugins/Experimental/Workspace/Source/WorkspaceEditor/Private/WorkspaceAssetRegistryInfo.cpp
@@ -2,6 +2,8 @@
 
 #include "WorkspaceAssetRegistryInfo.h"
 
+#include UE_INLINE_GENERATED_CPP_BY_NAME(WorkspaceAssetRegistryInfo)
+
 FWorkspaceOutlinerItemExport& FWorkspaceOutlinerItemExport::GetResolvedExport()
 {
 	if (HasData() && Data.GetScriptStruct() == FWorkspaceOutlinerAssetReferenceItemData::StaticStruct())
diff --git a/Engine/Plugins/Experimental/Workspace/Source/WorkspaceEditor/Private/WorkspaceFactory.cpp b/Engine/Plugins/Experimental/Workspace/Source/WorkspaceEditor/Private/WorkspaceFactory.cpp
--- a/Engine/Plugins/Experimental/Workspace/Source/WorkspaceEditor/Private/WorkspaceFactory.cpp
+++ b/Engine/Plugins/Experimental/Workspace/Source/WorkspaceEditor/Private/WorkspaceFactory.cpp
@@ -9,6 +9,8 @@
 #include "Kismet2/SClassPickerDialog.h"
 #include "UObject/Package.h"
 
+#include UE_INLINE_GENERATED_CPP_BY_NAME(WorkspaceFactory)
+
 #define LOCTEXT_NAMESPACE "WorkspaceFactory"
 
 UWorkspaceFactory::UWorkspaceFactory()
@@ -67,4 +69,4 @@ UObject* UWorkspaceFactory::FactoryCreateNew(UClass* Class, UObject* InParent, F
 	return NewWorkspace;
 }
 
-#undef LOCTEXT_NAMESPACE
\ No newline at end of file
+#undef LOCTEXT_NAMESPACE
diff --git a/Engine/Plugins/Experimental/Workspace/Source/WorkspaceEditor/Private/WorkspaceState.cpp b/Engine/Plugins/Experimental/Workspace/Source/WorkspaceEditor/Private/WorkspaceState.cpp
--- a/Engine/Plugins/Experimental/Workspace/Source/WorkspaceEditor/Private/WorkspaceState.cpp
+++ b/Engine/Plugins/Experimental/Workspace/Source/WorkspaceEditor/Private/WorkspaceState.cpp
@@ -1,4 +1,4 @@
-﻿// Copyright Epic Games, Inc. All Rights Reserved.
+// Copyright Epic Games, Inc. All Rights Reserved.
 
 #include "WorkspaceState.h"
 
@@ -11,6 +11,8 @@
 #include "Misc/Paths.h"
 #include "Modules/ModuleManager.h"
 
+#include UE_INLINE_GENERATED_CPP_BY_NAME(WorkspaceState)
+
 UWorkspaceState::UWorkspaceState()
 {
 	if(!HasAnyFlags(RF_ClassDefaultObject))
@@ -70,4 +72,4 @@ void UWorkspaceState::HandleAssetRenamed(const FAssetData& InAssetData, const FS
 			DocumentPath = InAssetData.GetSoftObjectPath();
 		}
 	}
-}
\ No newline at end of file
+}
diff --git a/Engine/Plugins/FX/Niagara/Source/Niagara/Private/NiagaraEmitterBase.cpp b/Engine/Plugins/FX/Niagara/Source/Niagara/Private/NiagaraEmitterBase.cpp
--- a/Engine/Plugins/FX/Niagara/Source/Niagara/Private/NiagaraEmitterBase.cpp
+++ b/Engine/Plugins/FX/Niagara/Source/Niagara/Private/NiagaraEmitterBase.cpp
@@ -5,6 +5,8 @@
 
 ////////////////////////////////////////////////////////////////////////////////////////////////////////////////////////////////////////////////////
 
+#include UE_INLINE_GENERATED_CPP_BY_NAME(NiagaraEmitterBase)
+
 FVersionedNiagaraEmitterBase::FVersionedNiagaraEmitterBase(const FVersionedNiagaraEmitterBaseWeakPtr& VersionedEmitterWeak)
 {
 	Emitter = VersionedEmitterWeak.Emitter.Get();
diff --git a/Engine/Plugins/FX/Niagara/Source/Niagara/Private/Stateless/Modules/NiagaraStatelessModule_AccelerationForce.cpp b/Engine/Plugins/FX/Niagara/Source/Niagara/Private/Stateless/Modules/NiagaraStatelessModule_AccelerationForce.cpp
--- a/Engine/Plugins/FX/Niagara/Source/Niagara/Private/Stateless/Modules/NiagaraStatelessModule_AccelerationForce.cpp
+++ b/Engine/Plugins/FX/Niagara/Source/Niagara/Private/Stateless/Modules/NiagaraStatelessModule_AccelerationForce.cpp
@@ -3,6 +3,8 @@
 #include "Stateless/Modules/NiagaraStatelessModule_AccelerationForce.h"
 #include "Stateless/NiagaraStatelessEmitterDataBuildContext.h"
 
+#include UE_INLINE_GENERATED_CPP_BY_NAME(NiagaraStatelessModule_AccelerationForce)
+
 void UNiagaraStatelessModule_AccelerationForce::BuildEmitterData(const FNiagaraStatelessEmitterDataBuildContext& BuildContext) const
 {
 	if (!IsModuleEnabled())
diff --git a/Engine/Plugins/FX/Niagara/Source/Niagara/Private/Stateless/Modules/NiagaraStatelessModule_CameraOffset.cpp b/Engine/Plugins/FX/Niagara/Source/Niagara/Private/Stateless/Modules/NiagaraStatelessModule_CameraOffset.cpp
--- a/Engine/Plugins/FX/Niagara/Source/Niagara/Private/Stateless/Modules/NiagaraStatelessModule_CameraOffset.cpp
+++ b/Engine/Plugins/FX/Niagara/Source/Niagara/Private/Stateless/Modules/NiagaraStatelessModule_CameraOffset.cpp
@@ -6,6 +6,8 @@
 #include "Stateless/NiagaraStatelessModuleShaderParameters.h"
 #include "Stateless/NiagaraStatelessParticleSimContext.h"
 
+#include UE_INLINE_GENERATED_CPP_BY_NAME(NiagaraStatelessModule_CameraOffset)
+
 namespace NSMCameraOffsetPrivate
 {
 	struct FModuleBuiltData
diff --git a/Engine/Plugins/FX/Niagara/Source/Niagara/Private/Stateless/Modules/NiagaraStatelessModule_Drag.cpp b/Engine/Plugins/FX/Niagara/Source/Niagara/Private/Stateless/Modules/NiagaraStatelessModule_Drag.cpp
--- a/Engine/Plugins/FX/Niagara/Source/Niagara/Private/Stateless/Modules/NiagaraStatelessModule_Drag.cpp
+++ b/Engine/Plugins/FX/Niagara/Source/Niagara/Private/Stateless/Modules/NiagaraStatelessModule_Drag.cpp
@@ -3,6 +3,8 @@
 #include "Stateless/Modules/NiagaraStatelessModule_Drag.h"
 #include "Stateless/NiagaraStatelessEmitterDataBuildContext.h"
 
+#include UE_INLINE_GENERATED_CPP_BY_NAME(NiagaraStatelessModule_Drag)
+
 void UNiagaraStatelessModule_Drag::BuildEmitterData(const FNiagaraStatelessEmitterDataBuildContext& BuildContext) const
 {
 	if (!IsModuleEnabled())
diff --git a/Engine/Plugins/FX/Niagara/Source/Niagara/Private/Stateless/Modules/NiagaraStatelessModule_DynamicMaterialParameters.cpp b/Engine/Plugins/FX/Niagara/Source/Niagara/Private/Stateless/Modules/NiagaraStatelessModule_DynamicMaterialParameters.cpp
--- a/Engine/Plugins/FX/Niagara/Source/Niagara/Private/Stateless/Modules/NiagaraStatelessModule_DynamicMaterialParameters.cpp
+++ b/Engine/Plugins/FX/Niagara/Source/Niagara/Private/Stateless/Modules/NiagaraStatelessModule_DynamicMaterialParameters.cpp
@@ -6,6 +6,8 @@
 #include "Stateless/NiagaraStatelessModuleShaderParameters.h"
 #include "Stateless/NiagaraStatelessParticleSimContext.h"
 
+#include UE_INLINE_GENERATED_CPP_BY_NAME(NiagaraStatelessModule_DynamicMaterialParameters)
+
 namespace NSMDynamicMaterialParametersPrivate
 {
 	static constexpr int32 NumParameters = 4;
diff --git a/Engine/Plugins/FX/Niagara/Source/Niagara/Private/Stateless/Modules/NiagaraStatelessModule_GravityForce.cpp b/Engine/Plugins/FX/Niagara/Source/Niagara/Private/Stateless/Modules/NiagaraStatelessModule_GravityForce.cpp
--- a/Engine/Plugins/FX/Niagara/Source/Niagara/Private/Stateless/Modules/NiagaraStatelessModule_GravityForce.cpp
+++ b/Engine/Plugins/FX/Niagara/Source/Niagara/Private/Stateless/Modules/NiagaraStatelessModule_GravityForce.cpp
@@ -5,6 +5,8 @@
 #include "Stateless/NiagaraStatelessEmitterDataBuildContext.h"
 #include "Stateless/Modules/NiagaraStatelessModuleCommon.h"
 
+#include UE_INLINE_GENERATED_CPP_BY_NAME(NiagaraStatelessModule_GravityForce)
+
 void UNiagaraStatelessModule_GravityForce::BuildEmitterData(const FNiagaraStatelessEmitterDataBuildContext& BuildContext) const
 {
 	if (!IsModuleEnabled())
diff --git a/Engine/Plugins/FX/Niagara/Source/Niagara/Private/Stateless/Modules/NiagaraStatelessModule_InitializeParticle.cpp b/Engine/Plugins/FX/Niagara/Source/Niagara/Private/Stateless/Modules/NiagaraStatelessModule_InitializeParticle.cpp
--- a/Engine/Plugins/FX/Niagara/Source/Niagara/Private/Stateless/Modules/NiagaraStatelessModule_InitializeParticle.cpp
+++ b/Engine/Plugins/FX/Niagara/Source/Niagara/Private/Stateless/Modules/NiagaraStatelessModule_InitializeParticle.cpp
@@ -10,6 +10,8 @@
 #include "NiagaraParameterBinding.h"
 #include "NiagaraParameterStore.h"
 
+#include UE_INLINE_GENERATED_CPP_BY_NAME(NiagaraStatelessModule_InitializeParticle)
+
 namespace NSMInitializeParticlePrivate
 {
 	using FParameters = NiagaraStateless::FInitializeParticleModule_ShaderParameters;
diff --git a/Engine/Plugins/FX/Niagara/Source/Niagara/Private/Stateless/Modules/NiagaraStatelessModule_MeshRotationRate.cpp b/Engine/Plugins/FX/Niagara/Source/Niagara/Private/Stateless/Modules/NiagaraStatelessModule_MeshRotationRate.cpp
--- a/Engine/Plugins/FX/Niagara/Source/Niagara/Private/Stateless/Modules/NiagaraStatelessModule_MeshRotationRate.cpp
+++ b/Engine/Plugins/FX/Niagara/Source/Niagara/Private/Stateless/Modules/NiagaraStatelessModule_MeshRotationRate.cpp
@@ -6,6 +6,8 @@
 #include "Stateless/NiagaraStatelessModuleShaderParameters.h"
 #include "Stateless/NiagaraStatelessParticleSimContext.h"
 
+#include UE_INLINE_GENERATED_CPP_BY_NAME(NiagaraStatelessModule_MeshRotationRate)
+
 namespace NSMMeshRotationRatePrivate
 {
 	using FParameters = NiagaraStateless::FMeshRotationRateModule_ShaderParameters;
diff --git a/Engine/Plugins/FX/Niagara/Source/Niagara/Private/Stateless/Modules/NiagaraStatelessModule_ScaleColor.cpp b/Engine/Plugins/FX/Niagara/Source/Niagara/Private/Stateless/Modules/NiagaraStatelessModule_ScaleColor.cpp
--- a/Engine/Plugins/FX/Niagara/Source/Niagara/Private/Stateless/Modules/NiagaraStatelessModule_ScaleColor.cpp
+++ b/Engine/Plugins/FX/Niagara/Source/Niagara/Private/Stateless/Modules/NiagaraStatelessModule_ScaleColor.cpp
@@ -5,6 +5,8 @@
 #include "Stateless/NiagaraStatelessModuleShaderParameters.h"
 #include "Stateless/NiagaraStatelessParticleSimContext.h"
 
+#include UE_INLINE_GENERATED_CPP_BY_NAME(NiagaraStatelessModule_ScaleColor)
+
 namespace NSMScaleColorPrivate
 {
 	struct FModuleBuiltData
diff --git a/Engine/Plugins/FX/Niagara/Source/Niagara/Private/Stateless/Modules/NiagaraStatelessModule_ScaleMeshSize.cpp b/Engine/Plugins/FX/Niagara/Source/Niagara/Private/Stateless/Modules/NiagaraStatelessModule_ScaleMeshSize.cpp
--- a/Engine/Plugins/FX/Niagara/Source/Niagara/Private/Stateless/Modules/NiagaraStatelessModule_ScaleMeshSize.cpp
+++ b/Engine/Plugins/FX/Niagara/Source/Niagara/Private/Stateless/Modules/NiagaraStatelessModule_ScaleMeshSize.cpp
@@ -5,6 +5,8 @@
 #include "Stateless/NiagaraStatelessModuleShaderParameters.h"
 #include "Stateless/NiagaraStatelessParticleSimContext.h"
 
+#include UE_INLINE_GENERATED_CPP_BY_NAME(NiagaraStatelessModule_ScaleMeshSize)
+
 namespace NSMScaleMeshSizePrivate
 {
 	struct FModuleBuiltData
diff --git a/Engine/Plugins/FX/Niagara/Source/Niagara/Private/Stateless/Modules/NiagaraStatelessModule_ScaleRibbonWidth.cpp b/Engine/Plugins/FX/Niagara/Source/Niagara/Private/Stateless/Modules/NiagaraStatelessModule_ScaleRibbonWidth.cpp
--- a/Engine/Plugins/FX/Niagara/Source/Niagara/Private/Stateless/Modules/NiagaraStatelessModule_ScaleRibbonWidth.cpp
+++ b/Engine/Plugins/FX/Niagara/Source/Niagara/Private/Stateless/Modules/NiagaraStatelessModule_ScaleRibbonWidth.cpp
@@ -5,6 +5,8 @@
 #include "Stateless/NiagaraStatelessModuleShaderParameters.h"
 #include "Stateless/NiagaraStatelessParticleSimContext.h"
 
+#include UE_INLINE_GENERATED_CPP_BY_NAME(NiagaraStatelessModule_ScaleRibbonWidth)
+
 namespace NSMScaleRibbonWidthPrivate
 {
 	struct FModuleBuiltData
diff --git a/Engine/Plugins/FX/Niagara/Source/Niagara/Private/Stateless/Modules/NiagaraStatelessModule_ScaleSpriteSize.cpp b/Engine/Plugins/FX/Niagara/Source/Niagara/Private/Stateless/Modules/NiagaraStatelessModule_ScaleSpriteSize.cpp
--- a/Engine/Plugins/FX/Niagara/Source/Niagara/Private/Stateless/Modules/NiagaraStatelessModule_ScaleSpriteSize.cpp
+++ b/Engine/Plugins/FX/Niagara/Source/Niagara/Private/Stateless/Modules/NiagaraStatelessModule_ScaleSpriteSize.cpp
@@ -6,6 +6,8 @@
 #include "Stateless/NiagaraStatelessParticleSimContext.h"
 #include "NiagaraParameterBinding.h"
 
+#include UE_INLINE_GENERATED_CPP_BY_NAME(NiagaraStatelessModule_ScaleSpriteSize)
+
 namespace NSMScaleSpriteSizePrivate
 {
 	struct FModuleBuiltData
diff --git a/Engine/Plugins/FX/Niagara/Source/Niagara/Private/Stateless/Modules/NiagaraStatelessModule_SpriteFacingAndAlignment.cpp b/Engine/Plugins/FX/Niagara/Source/Niagara/Private/Stateless/Modules/NiagaraStatelessModule_SpriteFacingAndAlignment.cpp
--- a/Engine/Plugins/FX/Niagara/Source/Niagara/Private/Stateless/Modules/NiagaraStatelessModule_SpriteFacingAndAlignment.cpp
+++ b/Engine/Plugins/FX/Niagara/Source/Niagara/Private/Stateless/Modules/NiagaraStatelessModule_SpriteFacingAndAlignment.cpp
@@ -6,6 +6,8 @@
 #include "Stateless/NiagaraStatelessModuleShaderParameters.h"
 #include "Stateless/NiagaraStatelessParticleSimContext.h"
 
+#include UE_INLINE_GENERATED_CPP_BY_NAME(NiagaraStatelessModule_SpriteFacingAndAlignment)
+
 namespace NSMSpriteFacingAndAlignmentPrivate
 {
 	struct FModuleBuiltData
diff --git a/Engine/Plugins/FX/Niagara/Source/Niagara/Private/Stateless/Modules/NiagaraStatelessModule_SpriteRotationRate.cpp b/Engine/Plugins/FX/Niagara/Source/Niagara/Private/Stateless/Modules/NiagaraStatelessModule_SpriteRotationRate.cpp
--- a/Engine/Plugins/FX/Niagara/Source/Niagara/Private/Stateless/Modules/NiagaraStatelessModule_SpriteRotationRate.cpp
+++ b/Engine/Plugins/FX/Niagara/Source/Niagara/Private/Stateless/Modules/NiagaraStatelessModule_SpriteRotationRate.cpp
@@ -5,6 +5,8 @@
 #include "Stateless/NiagaraStatelessModuleShaderParameters.h"
 #include "Stateless/NiagaraStatelessParticleSimContext.h"
 
+#include UE_INLINE_GENERATED_CPP_BY_NAME(NiagaraStatelessModule_SpriteRotationRate)
+
 namespace NSMSpriteRotationRatePrivate
 {
 	struct FModuleBuiltData
diff --git a/Engine/Plugins/Media/MediaIOFramework/Source/MediaIOCore/Private/RenderTargetMediaCapture.cpp b/Engine/Plugins/Media/MediaIOFramework/Source/MediaIOCore/Private/RenderTargetMediaCapture.cpp
--- a/Engine/Plugins/Media/MediaIOFramework/Source/MediaIOCore/Private/RenderTargetMediaCapture.cpp
+++ b/Engine/Plugins/Media/MediaIOFramework/Source/MediaIOCore/Private/RenderTargetMediaCapture.cpp
@@ -9,6 +9,8 @@
 #include "Slate/SceneViewport.h"
 #include "TextureResource.h"
 
+#include UE_INLINE_GENERATED_CPP_BY_NAME(RenderTargetMediaCapture)
+
 
 void URenderTargetMediaCapture::OnRHIResourceCaptured_RenderingThread(
 	FRHICommandListImmediate& RHICmdList, 
diff --git a/Engine/Plugins/Media/MediaIOFramework/Source/MediaIOCore/Private/RenderTargetMediaOutput.cpp b/Engine/Plugins/Media/MediaIOFramework/Source/MediaIOCore/Private/RenderTargetMediaOutput.cpp
--- a/Engine/Plugins/Media/MediaIOFramework/Source/MediaIOCore/Private/RenderTargetMediaOutput.cpp
+++ b/Engine/Plugins/Media/MediaIOFramework/Source/MediaIOCore/Private/RenderTargetMediaOutput.cpp
@@ -6,6 +6,8 @@
 #include "MediaIOCoreModule.h"
 #include "RenderTargetMediaCapture.h"
 
+#include UE_INLINE_GENERATED_CPP_BY_NAME(RenderTargetMediaOutput)
+
 bool URenderTargetMediaOutput::Validate(FString& OutFailureReason) const
 {
 	if (!Super::Validate(OutFailureReason))
diff --git a/Engine/Plugins/MovieScene/MovieRenderPipeline/Source/MovieRenderPipelineEditor/Private/MoviePipelinePostRenderFileDisplayProcessor.cpp b/Engine/Plugins/MovieScene/MovieRenderPipeline/Source/MovieRenderPipelineEditor/Private/MoviePipelinePostRenderFileDisplayProcessor.cpp
--- a/Engine/Plugins/MovieScene/MovieRenderPipeline/Source/MovieRenderPipelineEditor/Private/MoviePipelinePostRenderFileDisplayProcessor.cpp
+++ b/Engine/Plugins/MovieScene/MovieRenderPipeline/Source/MovieRenderPipelineEditor/Private/MoviePipelinePostRenderFileDisplayProcessor.cpp
@@ -12,6 +12,8 @@
 #include "Settings/EditorLoadingSavingSettings.h"
 #include "UObject/Package.h"
 
+#include UE_INLINE_GENERATED_CPP_BY_NAME(MoviePipelinePostRenderFileDisplayProcessor)
+
 FMoviePipelinePostRenderFileDisplayProcessor::FMoviePipelinePostRenderFileDisplayProcessor(const FMovieGraphPostRenderSettings& InPostRenderSettings)
 	: PostRenderSettings(InPostRenderSettings)
 {
diff --git a/Engine/Plugins/Mutable/Source/CustomizableObjectEditor/Private/MuCOE/Nodes/CustomizableObjectNodeMaterialParameter.cpp b/Engine/Plugins/Mutable/Source/CustomizableObjectEditor/Private/MuCOE/Nodes/CustomizableObjectNodeMaterialParameter.cpp
--- a/Engine/Plugins/Mutable/Source/CustomizableObjectEditor/Private/MuCOE/Nodes/CustomizableObjectNodeMaterialParameter.cpp
+++ b/Engine/Plugins/Mutable/Source/CustomizableObjectEditor/Private/MuCOE/Nodes/CustomizableObjectNodeMaterialParameter.cpp
@@ -1,9 +1,11 @@
-﻿// Copyright Epic Games, Inc. All Rights Reserved.
+// Copyright Epic Games, Inc. All Rights Reserved.
 
 #include "CustomizableObjectNodeMaterialParameter.h"
 
 #include "MuCOE/EdGraphSchema_CustomizableObject.h"
 
+#include UE_INLINE_GENERATED_CPP_BY_NAME(CustomizableObjectNodeMaterialParameter)
+
 
 bool UCustomizableObjectNodeMaterialParameter::IsExperimental() const
 {
diff --git a/Engine/Plugins/PCG/Source/PCG/Private/Elements/PCGSplineIntersection.cpp b/Engine/Plugins/PCG/Source/PCG/Private/Elements/PCGSplineIntersection.cpp
--- a/Engine/Plugins/PCG/Source/PCG/Private/Elements/PCGSplineIntersection.cpp
+++ b/Engine/Plugins/PCG/Source/PCG/Private/Elements/PCGSplineIntersection.cpp
@@ -9,6 +9,8 @@
 #include "Helpers/PCGHelpers.h"
 #include "Metadata/Accessors/PCGAttributeAccessorHelpers.h"
 
+#include UE_INLINE_GENERATED_CPP_BY_NAME(PCGSplineIntersection)
+
 #define LOCTEXT_NAMESPACE "PCGSplineIntersectionElement"
 
 namespace PCGSplineIntersection
@@ -529,4 +531,4 @@ bool FPCGSplineIntersectionElement::ExecuteInternal(FPCGContext* Context) const
 	return true;
 }
 
-#undef LOCTEXT_NAMESPACE
\ No newline at end of file
+#undef LOCTEXT_NAMESPACE
diff --git a/Engine/Plugins/Runtime/ComputeFramework/Source/ComputeFrameworkEditor/Private/ComputeGraphFromTextFactory.cpp b/Engine/Plugins/Runtime/ComputeFramework/Source/ComputeFrameworkEditor/Private/ComputeGraphFromTextFactory.cpp
--- a/Engine/Plugins/Runtime/ComputeFramework/Source/ComputeFrameworkEditor/Private/ComputeGraphFromTextFactory.cpp
+++ b/Engine/Plugins/Runtime/ComputeFramework/Source/ComputeFrameworkEditor/Private/ComputeGraphFromTextFactory.cpp
@@ -4,6 +4,8 @@
 
 #include "ComputeFramework/ComputeGraphFromText.h"
 
+#include UE_INLINE_GENERATED_CPP_BY_NAME(ComputeGraphFromTextFactory)
+
 UComputeGraphFromTextFactory::UComputeGraphFromTextFactory()
 {
 	SupportedClass = UComputeGraphFromText::StaticClass();
diff --git a/Engine/Plugins/Runtime/GameInput/Source/GameInputBase/Private/GameInputDeveloperSettings.cpp b/Engine/Plugins/Runtime/GameInput/Source/GameInputBase/Private/GameInputDeveloperSettings.cpp
--- a/Engine/Plugins/Runtime/GameInput/Source/GameInputBase/Private/GameInputDeveloperSettings.cpp
+++ b/Engine/Plugins/Runtime/GameInput/Source/GameInputBase/Private/GameInputDeveloperSettings.cpp
@@ -5,6 +5,8 @@
 #include "GameInputLogging.h"
 #include "GameInputBaseIncludes.h"
 
+#include UE_INLINE_GENERATED_CPP_BY_NAME(GameInputDeveloperSettings)
+
 FGameInputDeviceIdentifier::FGameInputDeviceIdentifier()
 	: VendorId(0)
 	, ProductId(0)
@@ -169,4 +171,4 @@ void UGameInputDeveloperSettings::PostEditChangeProperty(FPropertyChangedEvent&
 		OnInputSettingChanged.Broadcast();
 	}
 }
-#endif	// WITH_EDITOR
\ No newline at end of file
+#endif	// WITH_EDITOR
diff --git a/Engine/Plugins/Runtime/ModelViewViewModel/Source/ModelViewViewModel/Private/Bindings/ConversionLibraries/MVVMFontConversionLibrary.cpp b/Engine/Plugins/Runtime/ModelViewViewModel/Source/ModelViewViewModel/Private/Bindings/ConversionLibraries/MVVMFontConversionLibrary.cpp
--- a/Engine/Plugins/Runtime/ModelViewViewModel/Source/ModelViewViewModel/Private/Bindings/ConversionLibraries/MVVMFontConversionLibrary.cpp
+++ b/Engine/Plugins/Runtime/ModelViewViewModel/Source/ModelViewViewModel/Private/Bindings/ConversionLibraries/MVVMFontConversionLibrary.cpp
@@ -5,6 +5,8 @@
 #include "Materials/MaterialInterface.h"
 #include "Materials/MaterialInstanceDynamic.h"
 
+#include UE_INLINE_GENERATED_CPP_BY_NAME(MVVMFontConversionLibrary)
+
 FSlateFontInfo UMVVMFontConversionLibrary::Conv_SetScalarParameter(FSlateFontInfo Font, FName ParameterName, float Value)
 {
 	if (UMaterialInstanceDynamic* DynamicMaterial = TryGetDynamicMaterial(Font, nullptr))
@@ -93,4 +95,4 @@ UMaterialInstanceDynamic* UMVVMFontConversionLibrary::TryGetDynamicMaterial(FSla
 	 }
 
 	return nullptr;
-}
\ No newline at end of file
+}
diff --git a/Engine/Plugins/Runtime/StateTree/Source/StateTreeDeveloper/Private/Debugger/SCompactTreeDebuggerView.cpp b/Engine/Plugins/Runtime/StateTree/Source/StateTreeDeveloper/Private/Debugger/SCompactTreeDebuggerView.cpp
--- a/Engine/Plugins/Runtime/StateTree/Source/StateTreeDeveloper/Private/Debugger/SCompactTreeDebuggerView.cpp
+++ b/Engine/Plugins/Runtime/StateTree/Source/StateTreeDeveloper/Private/Debugger/SCompactTreeDebuggerView.cpp
@@ -7,6 +7,8 @@
 #include "StateTreeStyle.h"
 #include "Widgets/Views/SListView.h"
 
+#include UE_INLINE_GENERATED_CPP_BY_NAME(SCompactTreeDebuggerView)
+
 namespace UE::StateTree
 {
 
diff --git a/Engine/Plugins/Runtime/StateTree/Source/StateTreeDeveloper/Private/Widgets/SCompactTreeView.cpp b/Engine/Plugins/Runtime/StateTree/Source/StateTreeDeveloper/Private/Widgets/SCompactTreeView.cpp
--- a/Engine/Plugins/Runtime/StateTree/Source/StateTreeDeveloper/Private/Widgets/SCompactTreeView.cpp
+++ b/Engine/Plugins/Runtime/StateTree/Source/StateTreeDeveloper/Private/Widgets/SCompactTreeView.cpp
@@ -10,6 +10,8 @@
 #include "Widgets/Text/STextBlock.h"
 #include "Widgets/Views/STreeView.h"
 
+#include UE_INLINE_GENERATED_CPP_BY_NAME(SCompactTreeView)
+
 #define LOCTEXT_NAMESPACE "SCompactTreeView"
 
 namespace UE::StateTree
@@ -430,4 +432,4 @@ void SCompactTreeView::RestoreExpansionState()
 
 } // UE::StateTree
 
-#undef LOCTEXT_NAMESPACE
\ No newline at end of file
+#undef LOCTEXT_NAMESPACE
diff --git a/Engine/Plugins/Runtime/StateTree/Source/StateTreeEditorModule/Private/Customizations/Widgets/SCompactTreeEditorView.cpp b/Engine/Plugins/Runtime/StateTree/Source/StateTreeEditorModule/Private/Customizations/Widgets/SCompactTreeEditorView.cpp
--- a/Engine/Plugins/Runtime/StateTree/Source/StateTreeEditorModule/Private/Customizations/Widgets/SCompactTreeEditorView.cpp
+++ b/Engine/Plugins/Runtime/StateTree/Source/StateTreeEditorModule/Private/Customizations/Widgets/SCompactTreeEditorView.cpp
@@ -9,6 +9,8 @@
 #include "Widgets/Views/SListView.h"
 #include "Widgets/Views/STableRow.h"
 
+#include UE_INLINE_GENERATED_CPP_BY_NAME(SCompactTreeEditorView)
+
 #define LOCTEXT_NAMESPACE "OutlinerStateTreeView"
 
 namespace UE::StateTree
@@ -333,4 +335,4 @@ FReply SCompactTreeEditorView::HandleAcceptDrop(const FDragDropEvent& DragDropEv
 
 } // UE::StateTree
 
-#undef LOCTEXT_NAMESPACE
\ No newline at end of file
+#undef LOCTEXT_NAMESPACE
diff --git a/Engine/Source/Editor/DataHierarchyEditor/Private/DataHierarchyEditorMisc.cpp b/Engine/Source/Editor/DataHierarchyEditor/Private/DataHierarchyEditorMisc.cpp
--- a/Engine/Source/Editor/DataHierarchyEditor/Private/DataHierarchyEditorMisc.cpp
+++ b/Engine/Source/Editor/DataHierarchyEditor/Private/DataHierarchyEditorMisc.cpp
@@ -1,7 +1,9 @@
-﻿// Copyright Epic Games, Inc. All Rights Reserved.
+// Copyright Epic Games, Inc. All Rights Reserved.
 
 #include "DataHierarchyEditorMisc.h"
 
+#include UE_INLINE_GENERATED_CPP_BY_NAME(DataHierarchyEditorMisc)
+
 FName UE::DataHierarchyEditor::GetUniqueName(FName CandidateName, const TSet<FName>& ExistingNames)
 {
 	// This utility function needs to generate a unique name while only considering the text portion of the name and
diff --git a/Engine/Source/Editor/UnrealEd/Private/Settings/EditorViewportSettings.cpp b/Engine/Source/Editor/UnrealEd/Private/Settings/EditorViewportSettings.cpp
--- a/Engine/Source/Editor/UnrealEd/Private/Settings/EditorViewportSettings.cpp
+++ b/Engine/Source/Editor/UnrealEd/Private/Settings/EditorViewportSettings.cpp
@@ -2,6 +2,8 @@
 
 #include "Settings/EditorViewportSettings.h"
 
+#include UE_INLINE_GENERATED_CPP_BY_NAME(EditorViewportSettings)
+
 namespace UE::Editor::Private
 {
 	constexpr float SMALLEST_SPEED = 0.0001f;
diff --git a/Engine/Source/Runtime/CoreUObject/Private/UObject/Class.cpp b/Engine/Source/Runtime/CoreUObject/Private/UObject/Class.cpp
--- a/Engine/Source/Runtime/CoreUObject/Private/UObject/Class.cpp
+++ b/Engine/Source/Runtime/CoreUObject/Private/UObject/Class.cpp
@@ -81,6 +81,8 @@
 // This flag enables some expensive class tree validation that is meant to catch mutations of 
 // the class tree outside of SetSuperStruct. It has been disabled because loading blueprints 
 // does a lot of mutation of the class tree, and the validation checks impact iteration time.
+
+#include UE_INLINE_GENERATED_CPP_BY_NAME(Class)
 #define DO_CLASS_TREE_VALIDATION 0
 
 DEFINE_LOG_CATEGORY(LogScriptSerialization);
diff --git a/Engine/Source/Runtime/CoreUObject/Private/UObject/Package.cpp b/Engine/Source/Runtime/CoreUObject/Private/UObject/Package.cpp
--- a/Engine/Source/Runtime/CoreUObject/Private/UObject/Package.cpp
+++ b/Engine/Source/Runtime/CoreUObject/Private/UObject/Package.cpp
@@ -19,6 +19,8 @@
 #include "UObject/UObjectHash.h"
 #include "UObject/UObjectThreadContext.h"
 
+#include UE_INLINE_GENERATED_CPP_BY_NAME(Package)
+
 /*-----------------------------------------------------------------------------
 	UPackage.
 -----------------------------------------------------------------------------*/
diff --git a/Engine/Source/Runtime/CoreUObject/Private/UObject/PrimaryAssetId.cpp b/Engine/Source/Runtime/CoreUObject/Private/UObject/PrimaryAssetId.cpp
--- a/Engine/Source/Runtime/CoreUObject/Private/UObject/PrimaryAssetId.cpp
+++ b/Engine/Source/Runtime/CoreUObject/Private/UObject/PrimaryAssetId.cpp
@@ -7,6 +7,8 @@
 #include "UObject/PropertyPortFlags.h"
 #include "UObject/UnrealType.h"
 
+#include UE_INLINE_GENERATED_CPP_BY_NAME(PrimaryAssetId)
+
 bool FPrimaryAssetType::ExportTextItem(FString& ValueStr, FPrimaryAssetType const& DefaultValue, UObject* Parent, int32 PortFlags, UObject* ExportRootScope) const
 {
 	if (!(PortFlags & PPF_Delimited))
@@ -144,4 +146,4 @@ void SerializeForLog(FCbWriter& Writer, const FPrimaryAssetId& Value)
 	Writer.AddString(ANSITEXTVIEW("PrimaryAssetType"), WriteToUtf8String<256>(Value.PrimaryAssetType));
 	Writer.AddString(ANSITEXTVIEW("PrimaryAssetName"), WriteToUtf8String<256>(Value.PrimaryAssetName));
 	Writer.EndObject();
-}
\ No newline at end of file
+}
diff --git a/Engine/Source/Runtime/CoreUObject/Private/UObject/TemplateString.cpp b/Engine/Source/Runtime/CoreUObject/Private/UObject/TemplateString.cpp
--- a/Engine/Source/Runtime/CoreUObject/Private/UObject/TemplateString.cpp
+++ b/Engine/Source/Runtime/CoreUObject/Private/UObject/TemplateString.cpp
@@ -1,10 +1,12 @@
-﻿// Copyright Epic Games, Inc. All Rights Reserved.
+// Copyright Epic Games, Inc. All Rights Reserved.
 
 #include "UObject/TemplateString.h"
 
 #include "Algo/Count.h"
 #include "Internationalization/TextFormatter.h"
 
+#include UE_INLINE_GENERATED_CPP_BY_NAME(TemplateString)
+
 bool FTemplateString::IsValid(const TArray<FString>& InValidArgs) const
 {
 	constexpr static TCHAR OpenBrace = '{';
diff --git a/Engine/Source/Runtime/Engine/Private/Model.cpp b/Engine/Source/Runtime/Engine/Private/Model.cpp
--- a/Engine/Source/Runtime/Engine/Private/Model.cpp
+++ b/Engine/Source/Runtime/Engine/Private/Model.cpp
@@ -14,6 +14,8 @@
 #include "UObject/GarbageCollectionSchema.h"
 #include "UObject/ObjectSaveContext.h"
 
+#include UE_INLINE_GENERATED_CPP_BY_NAME(Model)
+
 float UModel::BSPTexelScale = 100.0f;
 
 DEFINE_LOG_CATEGORY_STATIC(LogModel, Log, All);
diff --git a/Engine/Source/Runtime/Engine/Public/Rendering/SubstrateMaterialShared.cpp b/Engine/Source/Runtime/Engine/Public/Rendering/SubstrateMaterialShared.cpp
--- a/Engine/Source/Runtime/Engine/Public/Rendering/SubstrateMaterialShared.cpp
+++ b/Engine/Source/Runtime/Engine/Public/Rendering/SubstrateMaterialShared.cpp
@@ -3,6 +3,8 @@
 #include "SubstrateMaterialShared.h"
 #include "Materials/MaterialExpressionSubstrate.h"
 
+#include UE_INLINE_GENERATED_CPP_BY_NAME(SubstrateMaterialShared)
+
 IMPLEMENT_TYPE_LAYOUT(FSubstrateMaterialCompilationOutput);
 
 FSubstrateMaterialCompilationOutput::FSubstrateMaterialCompilationOutput()
@@ -101,4 +103,4 @@ FSubstrateRegisteredSharedLocalBasis::FSubstrateRegisteredSharedLocalBasis()
 	TangentCodeChunkHash = 0;
 	GraphSharedLocalBasisIndex = 0;
 #endif
-}
\ No newline at end of file
+}
diff --git a/Engine/Source/Runtime/Experimental/Dataflow/Engine/Private/Dataflow/DataflowColorRamp.cpp b/Engine/Source/Runtime/Experimental/Dataflow/Engine/Private/Dataflow/DataflowColorRamp.cpp
--- a/Engine/Source/Runtime/Experimental/Dataflow/Engine/Private/Dataflow/DataflowColorRamp.cpp
+++ b/Engine/Source/Runtime/Experimental/Dataflow/Engine/Private/Dataflow/DataflowColorRamp.cpp
@@ -2,6 +2,8 @@
 
 #include "Dataflow/DataflowColorRamp.h"
 
+#include UE_INLINE_GENERATED_CPP_BY_NAME(DataflowColorRamp)
+
 TArray<FRichCurveEditInfoConst> FDataflowColorCurveOwner::GetCurves() const
 {
 	return ConstCurves;
diff --git a/Engine/Source/Runtime/Experimental/Dataflow/Engine/Private/Dataflow/DataflowDynamicConnections.cpp b/Engine/Source/Runtime/Experimental/Dataflow/Engine/Private/Dataflow/DataflowDynamicConnections.cpp
--- a/Engine/Source/Runtime/Experimental/Dataflow/Engine/Private/Dataflow/DataflowDynamicConnections.cpp
+++ b/Engine/Source/Runtime/Experimental/Dataflow/Engine/Private/Dataflow/DataflowDynamicConnections.cpp
@@ -4,6 +4,8 @@
 
 #include "Dataflow/DataflowObject.h"
 
+#include UE_INLINE_GENERATED_CPP_BY_NAME(DataflowDynamicConnections)
+
 FDataflowDynamicConnections::FDataflowDynamicConnections()
 	: PinDirection(UE::Dataflow::FPin::EDirection::NONE)
 	, OwnerInterface(nullptr)
@@ -263,4 +265,4 @@ void FDataflowDynamicConnections::ClearDynamicConnections()
 	}
 	ConnectionNameToPropertyId.Reset();
 	DynamicProperties.Reset();
-}
\ No newline at end of file
+}
diff --git a/Engine/Source/Runtime/InteractiveToolsFramework/Private/BaseGizmos/GizmoElementShared.cpp b/Engine/Source/Runtime/InteractiveToolsFramework/Private/BaseGizmos/GizmoElementShared.cpp
--- a/Engine/Source/Runtime/InteractiveToolsFramework/Private/BaseGizmos/GizmoElementShared.cpp
+++ b/Engine/Source/Runtime/InteractiveToolsFramework/Private/BaseGizmos/GizmoElementShared.cpp
@@ -1,7 +1,9 @@
-﻿// Copyright Epic Games, Inc. All Rights Reserved.
+// Copyright Epic Games, Inc. All Rights Reserved.
 
 #include "BaseGizmos/GizmoElementShared.h"
 
+#include UE_INLINE_GENERATED_CPP_BY_NAME(GizmoElementShared)
+
 namespace UE::InteractiveToolsFramework::Private
 {
 	namespace GizmoElementSharedLocals
