UAF: Fix crash with mesh component bindings

RigVM cast nodes were using the incorrect class to perform the comparison between object class and target class (field class was being compared rather than UClass)

diff --git a/Engine/Plugins/Experimental/UAF/UAFAnimGraph/Source/UAFAnimGraph/Private/Graph/AnimNextAnimationGraph.cpp b/Engine/Plugins/Experimental/UAF/UAFAnimGraph/Source/UAFAnimGraph/Private/Graph/AnimNextAnimationGraph.cpp
--- a/Engine/Plugins/Experimental/UAF/UAFAnimGraph/Source/UAFAnimGraph/Private/Graph/AnimNextAnimationGraph.cpp
+++ b/Engine/Plugins/Experimental/UAF/UAFAnimGraph/Source/UAFAnimGraph/Private/Graph/AnimNextAnimationGraph.cpp
@@ -71,9 +71,12 @@ TSharedPtr<FAnimNextGraphInstance> UAnimNextAnimationGraph::AllocateInstance(con
 		{
 			const FAnimNextSkeletalMeshComponentReferenceComponent& ComponentReference = InParams.ModuleInstance->GetComponent<FAnimNextSkeletalMeshComponentReferenceComponent>();
 
-			UE::UAF::FDataHandle RefPoseHandle = UE::UAF::FDataRegistry::Get()->GetOrGenerateReferencePose(ComponentReference.GetComponent());
-			const UE::UAF::FReferencePose& RefPose = RefPoseHandle.GetRef<UE::UAF::FReferencePose>();
-			Context.SetBindingObject(RefPose.SkeletalMeshComponent);
+			if (USkeletalMeshComponent* Component = ComponentReference.GetComponent())
+			{
+				UE::UAF::FDataHandle RefPoseHandle = UE::UAF::FDataRegistry::Get()->GetOrGenerateReferencePose(Component);
+				const UE::UAF::FReferencePose& RefPose = RefPoseHandle.GetRef<UE::UAF::FReferencePose>();
+				Context.SetBindingObject(RefPose.SkeletalMeshComponent);
+			}
 		}
 		InstanceImpl->GraphInstancePtr = Context.AllocateNodeInstance(*InstanceImpl, ResolvedRootTraitHandle);
 	}
diff --git a/Engine/Plugins/Runtime/RigVM/Source/RigVM/Private/RigVMFunctions/RigVMDispatch_CastObject.cpp b/Engine/Plugins/Runtime/RigVM/Source/RigVM/Private/RigVMFunctions/RigVMDispatch_CastObject.cpp
--- a/Engine/Plugins/Runtime/RigVM/Source/RigVM/Private/RigVMFunctions/RigVMDispatch_CastObject.cpp
+++ b/Engine/Plugins/Runtime/RigVM/Source/RigVM/Private/RigVMFunctions/RigVMDispatch_CastObject.cpp
@@ -112,7 +112,7 @@ void FRigVMDispatch_CastObject::Execute(FRigVMExtendedExecuteContext& InContext,
 	const FObjectPropertyBase* ResultObjectProperty = CastFieldChecked<FObjectPropertyBase>(Handles[1].GetProperty());
 	if(ValueObjectPtr && ResultObjectPtr && ValueObjectProperty && ResultObjectProperty)
 	{
-		if(ResultObjectProperty->GetClass()->IsChildOf(ValueObjectProperty->GetClass()))
+		if(*ValueObjectPtr && (*ValueObjectPtr)->GetClass()->IsChildOf(ResultObjectProperty->PropertyClass))
 		{
 			*ResultObjectPtr = *ValueObjectPtr;
 		}
