UAF: Prevent warning about missing style image in non-editor builds

Icon used by default was contained in the starship editor style, which is not loaded in e.g. -game, cookers

diff --git a/Engine/Plugins/Experimental/UAF/UAFAnimGraph/Source/UAFAnimGraphUncookedOnly/Private/Templates/UAFGraphNodeTemplate.cpp b/Engine/Plugins/Experimental/UAF/UAFAnimGraph/Source/UAFAnimGraphUncookedOnly/Private/Templates/UAFGraphNodeTemplate.cpp
--- a/Engine/Plugins/Experimental/UAF/UAFAnimGraph/Source/UAFAnimGraphUncookedOnly/Private/Templates/UAFGraphNodeTemplate.cpp
+++ b/Engine/Plugins/Experimental/UAF/UAFAnimGraph/Source/UAFAnimGraphUncookedOnly/Private/Templates/UAFGraphNodeTemplate.cpp
@@ -5,6 +5,7 @@
 #include "AnimNextController.h"
 #include "AnimNextTraitStackUnitNode.h"
 #include "EditorUtils.h"
+#include "UAFAnimGraphUncookedOnlyStyle.h"
 #include "Graph/RigUnit_AnimNextTraitStack.h"
 #include "RigVMModel/Nodes/RigVMUnitNode.h"
 #include "Styling/AppStyle.h"
@@ -17,7 +18,7 @@ UUAFGraphNodeTemplate::UUAFGraphNodeTemplate()
 	TooltipText = GetClass()->GetDisplayNameText();
 	Category = LOCTEXT("DefaultCategory", "UAF");
 	MenuDescription = GetClass()->GetDisplayNameText();
-	Icon = *FAppStyle::GetBrush(TEXT("GraphEditor.Animation_16x"));
+	Icon = *FUAFAnimGraphUncookedOnlyStyle::Get().GetBrush(TEXT("NodeTemplate.DefaultIcon"));
 }
 
 bool UUAFGraphNodeTemplate::ConfigureNewNode_Implementation(UAnimNextController* Controller, URigVMUnitNode* Node) const
diff --git a/Engine/Plugins/Experimental/UAF/UAFAnimGraph/Source/UAFAnimGraphUncookedOnly/Private/UAFAnimGraphUncookedOnlyStyle.cpp b/Engine/Plugins/Experimental/UAF/UAFAnimGraph/Source/UAFAnimGraphUncookedOnly/Private/UAFAnimGraphUncookedOnlyStyle.cpp
new file mode 100644
--- /dev/null
+++ b/Engine/Plugins/Experimental/UAF/UAFAnimGraph/Source/UAFAnimGraphUncookedOnly/Private/UAFAnimGraphUncookedOnlyStyle.cpp
@@ -0,0 +1,44 @@
+﻿// Copyright Epic Games, Inc. All Rights Reserved.
+
+#include "UAFAnimGraphUncookedOnlyStyle.h"
+#include "Styling/SlateStyleRegistry.h"
+#include "Brushes/SlateImageBrush.h"
+#include "Styling/CoreStyle.h"
+#include "Interfaces/IPluginManager.h"
+#include "Styling/SlateTypes.h"
+#include "Misc/Paths.h"
+#include "Styling/SlateStyleMacros.h"
+#include "Styling/StyleColors.h"
+
+FUAFAnimGraphUncookedOnlyStyle::FUAFAnimGraphUncookedOnlyStyle()
+	: FSlateStyleSet(TEXT("UAFAnimGraphUncookedOnlyStyle"))
+{
+	SetContentRoot(FPaths::EngineContentDir() / TEXT("Editor/Slate"));
+	SetCoreContentRoot(FPaths::EngineContentDir() / TEXT("Slate"));
+
+	Set("NodeTemplate.DefaultIcon", new IMAGE_BRUSH_SVG("Starship/Common/Animation", FVector2f(16.0f, 16.0f)));
+
+	FSlateStyleRegistry::RegisterSlateStyle(*this);
+}
+
+void FUAFAnimGraphUncookedOnlyStyle::Register()
+{
+	FSlateStyleRegistry::RegisterSlateStyle(Get());
+}
+
+void FUAFAnimGraphUncookedOnlyStyle::Unregister()
+{
+	FSlateStyleRegistry::UnRegisterSlateStyle(Get());
+}
+
+FUAFAnimGraphUncookedOnlyStyle& FUAFAnimGraphUncookedOnlyStyle::Get()
+{
+	static FUAFAnimGraphUncookedOnlyStyle Instance;
+	return Instance;
+}
+
+FString FUAFAnimGraphUncookedOnlyStyle::InResources(const FString& RelativePath)
+{
+	static const FString ResourcesDir = IPluginManager::Get().FindPlugin(TEXT("UAFAnimGraph"))->GetBaseDir() / TEXT("Resources");
+	return ResourcesDir / RelativePath;
+}
\ No newline at end of file
diff --git a/Engine/Plugins/Experimental/UAF/UAFAnimGraph/Source/UAFAnimGraphUncookedOnly/Private/UAFAnimGraphUncookedOnlyStyle.h b/Engine/Plugins/Experimental/UAF/UAFAnimGraph/Source/UAFAnimGraphUncookedOnly/Private/UAFAnimGraphUncookedOnlyStyle.h
new file mode 100644
--- /dev/null
+++ b/Engine/Plugins/Experimental/UAF/UAFAnimGraph/Source/UAFAnimGraphUncookedOnly/Private/UAFAnimGraphUncookedOnlyStyle.h
@@ -0,0 +1,22 @@
+﻿// Copyright Epic Games, Inc. All Rights Reserved.
+
+#pragma once
+
+#include "Styling/SlateStyle.h"
+
+class FUAFAnimGraphUncookedOnlyStyle final : public FSlateStyleSet
+{
+public:
+	static FUAFAnimGraphUncookedOnlyStyle& Get();
+
+protected:
+	static void Register();
+	static void Unregister();
+
+private:
+	FUAFAnimGraphUncookedOnlyStyle();
+
+	// Resolve path relative to plugin's resources directory
+	static FString InResources(const FString& RelativePath);
+};
+
diff --git a/Engine/Plugins/Experimental/UAF/UAFAnimGraph/Source/UAFAnimGraphUncookedOnly/UAFAnimGraphUncookedOnly.Build.cs b/Engine/Plugins/Experimental/UAF/UAFAnimGraph/Source/UAFAnimGraphUncookedOnly/UAFAnimGraphUncookedOnly.Build.cs
--- a/Engine/Plugins/Experimental/UAF/UAFAnimGraph/Source/UAFAnimGraphUncookedOnly/UAFAnimGraphUncookedOnly.Build.cs
+++ b/Engine/Plugins/Experimental/UAF/UAFAnimGraph/Source/UAFAnimGraphUncookedOnly/UAFAnimGraphUncookedOnly.Build.cs
@@ -26,7 +26,8 @@ namespace UnrealBuildTool.Rules
 					"AnimGraph",
 					"RigVM",
 					"Slate",
-					"AssetDefinition"
+					"AssetDefinition",
+					"Projects",
 				}
 			);
 
