UAF: Fix UAF UncookedOnly issues with hot-reload

Register style on module load so it is not attempted to be re-registered on hot reload

diff --git a/Engine/Plugins/Experimental/UAF/UAFAnimGraph/Source/UAFAnimGraphUncookedOnly/Private/AnimNextAnimGraphUncookedOnlyModule.cpp b/Engine/Plugins/Experimental/UAF/UAFAnimGraph/Source/UAFAnimGraphUncookedOnly/Private/AnimNextAnimGraphUncookedOnlyModule.cpp
--- a/Engine/Plugins/Experimental/UAF/UAFAnimGraph/Source/UAFAnimGraphUncookedOnly/Private/AnimNextAnimGraphUncookedOnlyModule.cpp
+++ b/Engine/Plugins/Experimental/UAF/UAFAnimGraph/Source/UAFAnimGraphUncookedOnly/Private/AnimNextAnimGraphUncookedOnlyModule.cpp
@@ -1,6 +1,7 @@
 // Copyright Epic Games, Inc. All Rights Reserved.
 
 #include "CoreMinimal.h"
+#include "UAFAnimGraphUncookedOnlyStyle.h"
 #include "Modules/ModuleInterface.h"
 #include "Modules/ModuleManager.h"
 #include "Graph/RigVMTrait_AnimNextPublicVariablesUncookedOnly.h"
@@ -18,6 +19,7 @@ class FAnimNextAnimGraphUncookedOnlyModule : public IModuleInterface
 		using namespace UE::UAF::UncookedOnly;
 
 		FPublicVariablesImpl::Register();
+		FUAFAnimGraphUncookedOnlyStyle::Get();
 	}
 
 	virtual void ShutdownModule() override
