UAF: Fix style registation/unregistration

Hot reload could cause styles to be incorrectly re-registered

diff --git a/Engine/Plugins/Experimental/UAF/UAF/Source/UAFEditor/Internal/UAFStyle.h b/Engine/Plugins/Experimental/UAF/UAF/Source/UAFEditor/Internal/UAFStyle.h
--- a/Engine/Plugins/Experimental/UAF/UAF/Source/UAFEditor/Internal/UAFStyle.h
+++ b/Engine/Plugins/Experimental/UAF/UAF/Source/UAFEditor/Internal/UAFStyle.h
@@ -11,12 +11,9 @@ class FUAFStyle final : public FSlateStyleSet
 public:
 	static UE_API FUAFStyle& Get();
 
-protected:
-	static UE_API void Register();
-	static UE_API void Unregister();
-
 private:
 	UE_API FUAFStyle();
+	UE_API ~FUAFStyle();
 
 	// Resolve path relative to plugin's resources directory
 	static UE_API FString InResources(const FString& RelativePath);
diff --git a/Engine/Plugins/Experimental/UAF/UAF/Source/UAFEditor/Private/UAFStyle.cpp b/Engine/Plugins/Experimental/UAF/UAF/Source/UAFEditor/Private/UAFStyle.cpp
--- a/Engine/Plugins/Experimental/UAF/UAF/Source/UAFEditor/Private/UAFStyle.cpp
+++ b/Engine/Plugins/Experimental/UAF/UAF/Source/UAFEditor/Private/UAFStyle.cpp
@@ -29,6 +29,11 @@ FUAFStyle::FUAFStyle()
 	FSlateStyleRegistry::RegisterSlateStyle(*this);
 }
 
+FUAFStyle::~FUAFStyle()
+{
+	FSlateStyleRegistry::UnRegisterSlateStyle(*this);
+}
+
 void FUAFStyle::SetClassIcons()
 {
 	const FVector2D Icon16x16(16.0f, 16.0f);
@@ -103,16 +108,6 @@ void FUAFStyle::SetAssetWizardStyles()
 	Set("AssetWizard.TemplateFolderHoverSelectedBackground", new FSlateRoundedBoxBrush(FStyleColors::PrimaryHover, 4.0f));
 }
 
-void FUAFStyle::Register()
-{
-	FSlateStyleRegistry::RegisterSlateStyle(Get());
-}
-
-void FUAFStyle::Unregister()
-{
-	FSlateStyleRegistry::UnRegisterSlateStyle(Get());
-}
-
 FUAFStyle& FUAFStyle::Get()
 {
 	static FUAFStyle Instance;
diff --git a/Engine/Plugins/Experimental/UAF/UAFAnimGraph/Source/UAFAnimGraphUncookedOnly/Private/UAFAnimGraphUncookedOnlyStyle.cpp b/Engine/Plugins/Experimental/UAF/UAFAnimGraph/Source/UAFAnimGraphUncookedOnly/Private/UAFAnimGraphUncookedOnlyStyle.cpp
--- a/Engine/Plugins/Experimental/UAF/UAFAnimGraph/Source/UAFAnimGraphUncookedOnly/Private/UAFAnimGraphUncookedOnlyStyle.cpp
+++ b/Engine/Plugins/Experimental/UAF/UAFAnimGraph/Source/UAFAnimGraphUncookedOnly/Private/UAFAnimGraphUncookedOnlyStyle.cpp
@@ -21,14 +21,9 @@ FUAFAnimGraphUncookedOnlyStyle::FUAFAnimGraphUncookedOnlyStyle()
 	FSlateStyleRegistry::RegisterSlateStyle(*this);
 }
 
-void FUAFAnimGraphUncookedOnlyStyle::Register()
+FUAFAnimGraphUncookedOnlyStyle::~FUAFAnimGraphUncookedOnlyStyle()
 {
-	FSlateStyleRegistry::RegisterSlateStyle(Get());
-}
-
-void FUAFAnimGraphUncookedOnlyStyle::Unregister()
-{
-	FSlateStyleRegistry::UnRegisterSlateStyle(Get());
+	FSlateStyleRegistry::UnRegisterSlateStyle(*this);
 }
 
 FUAFAnimGraphUncookedOnlyStyle& FUAFAnimGraphUncookedOnlyStyle::Get()
diff --git a/Engine/Plugins/Experimental/UAF/UAFAnimGraph/Source/UAFAnimGraphUncookedOnly/Private/UAFAnimGraphUncookedOnlyStyle.h b/Engine/Plugins/Experimental/UAF/UAFAnimGraph/Source/UAFAnimGraphUncookedOnly/Private/UAFAnimGraphUncookedOnlyStyle.h
--- a/Engine/Plugins/Experimental/UAF/UAFAnimGraph/Source/UAFAnimGraphUncookedOnly/Private/UAFAnimGraphUncookedOnlyStyle.h
+++ b/Engine/Plugins/Experimental/UAF/UAFAnimGraph/Source/UAFAnimGraphUncookedOnly/Private/UAFAnimGraphUncookedOnlyStyle.h
@@ -9,12 +9,9 @@ class FUAFAnimGraphUncookedOnlyStyle final : public FSlateStyleSet
 public:
 	static FUAFAnimGraphUncookedOnlyStyle& Get();
 
-protected:
-	static void Register();
-	static void Unregister();
-
 private:
 	FUAFAnimGraphUncookedOnlyStyle();
+	~FUAFAnimGraphUncookedOnlyStyle();
 
 	// Resolve path relative to plugin's resources directory
 	static FString InResources(const FString& RelativePath);
