Properly handle out of order latent pins
Add UAF unit tests
* Test out of order latent pins
* Test FVector latent pin

diff --git a/Engine/Plugins/Experimental/UAF/UAFAnimGraph/Source/UAFAnimGraph/Private/TraitCore/Trait.cpp b/Engine/Plugins/Experimental/UAF/UAFAnimGraph/Source/UAFAnimGraph/Private/TraitCore/Trait.cpp
--- a/Engine/Plugins/Experimental/UAF/UAFAnimGraph/Source/UAFAnimGraph/Private/TraitCore/Trait.cpp
+++ b/Engine/Plugins/Experimental/UAF/UAFAnimGraph/Source/UAFAnimGraph/Private/TraitCore/Trait.cpp
@@ -13,10 +13,12 @@ namespace UE::UAF
 	}
 
 	FTraitLatentPropertyMemoryLayout FTrait::GetLatentPropertyMemoryLayoutImpl(
+		const FAnimNextTraitSharedData& SharedData,
 		FName PropertyName,
 		uint32 PropertyIndex,
 		TArray<FTraitLatentPropertyMemoryLayout>& LatentPropertyMemoryLayouts) const
 	{
+		// The property index here isn't important as long as it is stable, only used for caching purposes
 		check(LatentPropertyMemoryLayouts.IsValidIndex(PropertyIndex));
 		if (LatentPropertyMemoryLayouts[PropertyIndex].Size == 0)
 		{
@@ -28,6 +30,10 @@ namespace UE::UAF
 
 			LatentPropertyMemoryLayouts[PropertyIndex].Alignment = Property->GetMinAlignment();
 
+			// Latent property handles are in the order defined by the enumerator macro within the shared data,
+			// which means their index can differ from the one we serialize with which uses the UStruct order
+			LatentPropertyMemoryLayouts[PropertyIndex].LatentPropertyIndex = GetLatentPropertyIndex(SharedData, PropertyName);
+
 			// Ensure alignment is visible before we write the size to avoid torn reads
 			FPlatformMisc::MemoryBarrier();
 
diff --git a/Engine/Plugins/Experimental/UAF/UAFAnimGraph/Source/UAFAnimGraph/Private/TraitCore/TraitReader.cpp b/Engine/Plugins/Experimental/UAF/UAFAnimGraph/Source/UAFAnimGraph/Private/TraitCore/TraitReader.cpp
--- a/Engine/Plugins/Experimental/UAF/UAFAnimGraph/Source/UAFAnimGraph/Private/TraitCore/TraitReader.cpp
+++ b/Engine/Plugins/Experimental/UAF/UAFAnimGraph/Source/UAFAnimGraph/Private/TraitCore/TraitReader.cpp
@@ -166,7 +166,16 @@ namespace UE::UAF
 				const FTraitRegistryHandle TraitHandle = TraitTemplate.GetRegistryHandle();
 				const FTrait* Trait = TraitRegistry.Find(TraitHandle);
 
-				for (uint32 LatentPropertyIndex = 0; LatentPropertyIndex < NumLatentProperties; ++LatentPropertyIndex)
+				if (Trait == nullptr)
+				{
+					// Default initialize our latent property handles if we don't have a valid trait
+					for (uint32 LatentPropertyIndex = 0; LatentPropertyIndex < NumLatentProperties; ++LatentPropertyIndex)
+					{
+						LatentHandles[LatentPropertyIndex] = FLatentPropertyHandle();
+					}
+				}
+
+				for (uint32 SerializedLatentPropertyIndex = 0; SerializedLatentPropertyIndex < NumLatentProperties; ++SerializedLatentPropertyIndex)
 				{
 					FLatentPropertyMetadata Metadata;
 					*this << Metadata;
@@ -178,13 +187,14 @@ namespace UE::UAF
 
 					if (Trait != nullptr)
 					{
+						const FTraitLatentPropertyMemoryLayout PropertyMemoryLayout = Trait->GetLatentPropertyMemoryLayout(*SharedData, Metadata.Name, SerializedLatentPropertyIndex);
+						check(PropertyMemoryLayout.Size != 0);
+						check(PropertyMemoryLayout.Alignment != 0 && FMath::IsPowerOfTwo(PropertyMemoryLayout.Alignment));
+						check(PropertyMemoryLayout.LatentPropertyIndex != INDEX_NONE && (uint32)PropertyMemoryLayout.LatentPropertyIndex < NumLatentProperties);
+
 						// If this property is valid, setup our binding for it
 						if (Metadata.RigVMIndex != MAX_uint16)
 						{
-							const FTraitLatentPropertyMemoryLayout PropertyMemoryLayout = Trait->GetLatentPropertyMemoryLayout(*SharedData, Metadata.Name, LatentPropertyIndex);
-							check(PropertyMemoryLayout.Size != 0);
-							check(PropertyMemoryLayout.Alignment != 0 && FMath::IsPowerOfTwo(PropertyMemoryLayout.Alignment));
-
 							// Align our property
 							LatentPropertyOffset = Align(LatentPropertyOffset, PropertyMemoryLayout.Alignment);
 
@@ -201,9 +211,11 @@ namespace UE::UAF
 							// Consume the property size
 							LatentPropertyOffset += PropertyMemoryLayout.Size;
 						}
-					}
 
-					LatentHandles[LatentPropertyIndex] = FLatentPropertyHandle(RigVMIndex, CurrentLatentPropertyOffset, bCanFreeze, bOnBecomeRelevant);
+						// Latent property handles are in the order defined by the enumerator macro within the shared data,
+						// which means their index can differ from the one we serialize with which uses the UStruct order
+						LatentHandles[PropertyMemoryLayout.LatentPropertyIndex] = FLatentPropertyHandle(RigVMIndex, CurrentLatentPropertyOffset, bCanFreeze, bOnBecomeRelevant);
+					}
 				}
 
 				LatentHeader.bHasValidLatentProperties = bHasValidLatentProperties;
diff --git a/Engine/Plugins/Experimental/UAF/UAFAnimGraph/Source/UAFAnimGraph/Public/TraitCore/Trait.h b/Engine/Plugins/Experimental/UAF/UAFAnimGraph/Source/UAFAnimGraph/Public/TraitCore/Trait.h
--- a/Engine/Plugins/Experimental/UAF/UAFAnimGraph/Source/UAFAnimGraph/Public/TraitCore/Trait.h
+++ b/Engine/Plugins/Experimental/UAF/UAFAnimGraph/Source/UAFAnimGraph/Public/TraitCore/Trait.h
@@ -49,7 +49,7 @@ struct FRigVMPinInfoArray;
 #define ANIM_NEXT_IMPL_DECLARE_ANIM_TRAIT_LATENT_PROPERTY_SUPPORT() \
 	virtual uint32 GetNumLatentTraitProperties() const override { return -FSharedData::GetLatentPropertyIndex(~(size_t)0); } \
 	virtual UE::UAF::FTraitLatentPropertyMemoryLayout GetLatentPropertyMemoryLayout(const FAnimNextTraitSharedData& SharedData, FName PropertyName, uint32 PropertyIndex) const override; \
-	virtual bool IsPropertyLatent(const FAnimNextTraitSharedData& InSharedData, FName PropertyName) const override;
+	virtual int32 GetLatentPropertyIndex(const FAnimNextTraitSharedData& InSharedData, FName PropertyName) const override;
 
 // Helper macros
 // In the trait class declaration, this macro declares the Super alias and base functions we override
@@ -109,16 +109,16 @@ struct FRigVMPinInfoArray;
 	{ \
 		/* Thread safe cache initialization */ \
 		static TArray<UE::UAF::FTraitLatentPropertyMemoryLayout> CachedLatentPropertyMemoryLayouts = [this](){ TArray<UE::UAF::FTraitLatentPropertyMemoryLayout> Result; Result.SetNum(GetNumLatentTraitProperties()); return Result; }(); \
-		return GetLatentPropertyMemoryLayoutImpl(PropertyName, PropertyIndex, CachedLatentPropertyMemoryLayouts); \
+		return GetLatentPropertyMemoryLayoutImpl(SharedData, PropertyName, PropertyIndex, CachedLatentPropertyMemoryLayouts); \
 	}
 
 #define ANIM_NEXT_IMPL_DEFINE_ANIM_TRAIT_IS_PROPERTY_LATENT(TraitName) \
-	bool TraitName::IsPropertyLatent(const FAnimNextTraitSharedData& InSharedData, FName PropertyName) const \
+	int32 TraitName::GetLatentPropertyIndex(const FAnimNextTraitSharedData& InSharedData, FName PropertyName) const \
 	{ \
 		const UScriptStruct* SharedDataStruct = FSharedData::StaticStruct(); \
 		const FProperty* Property = SharedDataStruct->FindPropertyByName(PropertyName); \
-		check(Property != nullptr); \
-		return FSharedData::GetLatentPropertyIndex(Property->GetOffset_ForInternal()) > 0; \
+		ensure(Property != nullptr); \
+		return Property != nullptr ? (FSharedData::GetLatentPropertyIndex(Property->GetOffset_ForInternal()) - 1) : INDEX_NONE; \
 	}
 
 // Helper that handles the GetTraitInterface() details for each interface specified by the generator macro
@@ -328,6 +328,9 @@ namespace UE::UAF
 
 		// The alignment in bytes of the latent property
 		uint32 Alignment = 1;
+
+		// The latent property index that matches the enumerator macro, @see FTrait::GetLatentPropertyIndex
+		int32 LatentPropertyIndex = INDEX_NONE;
 	};
 
 	/**
@@ -463,11 +466,17 @@ namespace UE::UAF
 			const TFunction<uint16(FName PropertyName)>& GetTraitLatentPropertyIndex) const;
 
 		// Returns whether or not the specified property has been marked as latent
-		virtual bool IsPropertyLatent(const FAnimNextTraitSharedData& InSharedData, FName PropertyName) const { return false; }
+		bool IsPropertyLatent(const FAnimNextTraitSharedData& InSharedData, FName PropertyName) const { return GetLatentPropertyIndex(InSharedData, PropertyName) != INDEX_NONE; }
+
+		// Returns the latent property index or INDEX_NONE if the property isn't found or if it isn't latent
+		// WARNING: Latent property handles are in the order defined by the enumerator macro within the shared data,
+		// @see FAnimNextTraitSharedData::GetLatentPropertyIndex
+		virtual int32 GetLatentPropertyIndex(const FAnimNextTraitSharedData& InSharedData, FName PropertyName) const { return INDEX_NONE; }
 		
 	protected:
 		// Implements GetLatentPropertyMemoryLayout() by allowing a map to be provided for caching purposes to speed up repeated queries
 		UE_API FTraitLatentPropertyMemoryLayout GetLatentPropertyMemoryLayoutImpl(
+			const FAnimNextTraitSharedData& SharedData,
 			FName PropertyName,
 			uint32 PropertyIndex,
 			TArray<FTraitLatentPropertyMemoryLayout>& LatentPropertyMemoryLayouts) const;
diff --git a/Engine/Plugins/Experimental/UAF/UAFAnimGraph/Source/UAFAnimGraph/Public/TraitCore/TraitBinding.h b/Engine/Plugins/Experimental/UAF/UAFAnimGraph/Source/UAFAnimGraph/Public/TraitCore/TraitBinding.h
--- a/Engine/Plugins/Experimental/UAF/UAFAnimGraph/Source/UAFAnimGraph/Public/TraitCore/TraitBinding.h
+++ b/Engine/Plugins/Experimental/UAF/UAFAnimGraph/Source/UAFAnimGraph/Public/TraitCore/TraitBinding.h
@@ -86,6 +86,8 @@ namespace UE::UAF
 
 		// Queries a node for a pointer to its trait latent properties.
 		// If the trait handle is invalid or if we have no latent properties, a null pointer is returned.
+		// WARNING: Latent property handles are in the order defined by the enumerator macro within the shared data,
+		// @see FAnimNextTraitSharedData::GetLatentPropertyIndex
 		const FLatentPropertyHandle* GetLatentPropertyHandles() const
 		{
 			if (!IsValid())
diff --git a/Engine/Plugins/Experimental/UAF/UAFAnimGraph/Source/UAFAnimGraph/Public/TraitCore/TraitSharedData.h b/Engine/Plugins/Experimental/UAF/UAFAnimGraph/Source/UAFAnimGraph/Public/TraitCore/TraitSharedData.h
--- a/Engine/Plugins/Experimental/UAF/UAFAnimGraph/Source/UAFAnimGraph/Public/TraitCore/TraitSharedData.h
+++ b/Engine/Plugins/Experimental/UAF/UAFAnimGraph/Source/UAFAnimGraph/Public/TraitCore/TraitSharedData.h
@@ -147,6 +147,8 @@ struct FAnimNextTraitSharedData
 	  * If the property is found within this type, its latent property 1-based index is returned (a positive value)
 	  * If the property isn't found, this function returns the number of latent properties (as a zero or negative value)
 	  * Derived types will have latent property indices higher than their base type
+	  * Note that the property index is derived from the order specified in the GENERATE_TRAIT_LATENT_PROPERTIES enumerator macro
+	  * which might not match the order of the UStruct
 	  */
 	static constexpr int32 GetLatentPropertyIndex(const size_t LatentPropertyOffset) { return 0; }
 };
diff --git a/Engine/Plugins/Experimental/UAF/UAFAnimGraph/Source/UAFAnimGraph/Public/TraitCore/TraitTemplate.h b/Engine/Plugins/Experimental/UAF/UAFAnimGraph/Source/UAFAnimGraph/Public/TraitCore/TraitTemplate.h
--- a/Engine/Plugins/Experimental/UAF/UAFAnimGraph/Source/UAFAnimGraph/Public/TraitCore/TraitTemplate.h
+++ b/Engine/Plugins/Experimental/UAF/UAFAnimGraph/Source/UAFAnimGraph/Public/TraitCore/TraitTemplate.h
@@ -106,12 +106,16 @@ namespace UE::UAF
 		}
 
 		// Returns a pointer to the specified trait latent property handles on the current node
+		// WARNING: Latent property handles are in the order defined by the enumerator macro within the shared data,
+		// @see FAnimNextTraitSharedData::GetLatentPropertyIndex
 		FLatentPropertyHandle* GetTraitLatentPropertyHandles(FNodeDescription& NodeDescription) const noexcept
 		{
 			return reinterpret_cast<FLatentPropertyHandle*>(reinterpret_cast<uint8*>(&NodeDescription) + GetNodeSharedLatentPropertyHandlesOffset());
 		}
 
 		// Returns a pointer to the specified trait latent property handles on the current node
+		// WARNING: Latent property handles are in the order defined by the enumerator macro within the shared data,
+		// @see FAnimNextTraitSharedData::GetLatentPropertyIndex
 		const FLatentPropertyHandle* GetTraitLatentPropertyHandles(const FNodeDescription& NodeDescription) const noexcept
 		{
 			return reinterpret_cast<const FLatentPropertyHandle*>(reinterpret_cast<const uint8*>(&NodeDescription) + GetNodeSharedLatentPropertyHandlesOffset());
diff --git a/Engine/Plugins/Experimental/UAF/UAFAnimGraph/Source/UAFAnimGraphTestSuite/Private/AnimNextAnimGraphTraitGraphTest.cpp b/Engine/Plugins/Experimental/UAF/UAFAnimGraph/Source/UAFAnimGraphTestSuite/Private/AnimNextAnimGraphTraitGraphTest.cpp
--- a/Engine/Plugins/Experimental/UAF/UAFAnimGraph/Source/UAFAnimGraphTestSuite/Private/AnimNextAnimGraphTraitGraphTest.cpp
+++ b/Engine/Plugins/Experimental/UAF/UAFAnimGraph/Source/UAFAnimGraphTestSuite/Private/AnimNextAnimGraphTraitGraphTest.cpp
@@ -63,6 +63,9 @@ namespace UE::UAF
 			ExecuteContext.Logf(EMessageSeverity::Info, TEXT("SomeFloat == %.02f"), SharedData->SomeFloat);
 			ExecuteContext.Logf(EMessageSeverity::Info, TEXT("SomeLatentInt32 == %d"), SharedData->GetSomeLatentInt32(Binding));
 			ExecuteContext.Logf(EMessageSeverity::Info, TEXT("SomeOtherLatentInt32 == %d"), SharedData->GetSomeOtherLatentInt32(Binding));
+			ExecuteContext.Logf(EMessageSeverity::Info, TEXT("SomeOutOfOrderLatentBool == %d"), (int32)SharedData->GetSomeOutOfOrderLatentBool(Binding));
+			const FVector& LatentVector = SharedData->GetSomeLatentVector(Binding);
+			ExecuteContext.Logf(EMessageSeverity::Info, TEXT("SomeLatentVector == [%.02f, %.02f, %.02f]"), LatentVector.X, LatentVector.Y, LatentVector.Z);
 			ExecuteContext.Logf(EMessageSeverity::Info, TEXT("SomeLatentFloat == %.02f"), SharedData->GetSomeLatentFloat(Binding));
 		}
 
@@ -165,7 +168,7 @@ bool FAnimationAnimNextEditorTest_GraphAddTrait::RunTest(const FString& InParame
 		UE_RETURN_ON_ERROR(TraitPin->GetCPPTypeObject() == FRigDecorator_AnimNextCppDecorator::StaticStruct(), TEXT("FAnimationAnimNextEditorTest_GraphAddTrait -> Unexpected pin type"));
 
 		// Our first sub-pin is the hard coded script struct member that parametrizes the trait, dynamic trait sub-pins follow
-		UE_RETURN_ON_ERROR(TraitPin->GetSubPins().Num() == 8, TEXT("FAnimationAnimNextEditorTest_GraphAddTrait -> Unexpected trait sub pins"));
+		UE_RETURN_ON_ERROR(TraitPin->GetSubPins().Num() == 10, TEXT("FAnimationAnimNextEditorTest_GraphAddTrait -> Unexpected trait sub pins"));
 
 		// UpdateCount
 		UE_RETURN_ON_ERROR(TraitPin->GetSubPins()[1]->GetCPPType() == TEXT("int32"), TEXT("FAnimationAnimNextEditorTest_GraphAddTrait -> Unexpected trait pin type"));
@@ -189,18 +192,28 @@ bool FAnimationAnimNextEditorTest_GraphAddTrait::RunTest(const FString& InParame
 
 		// SomeLatentInt32
 		UE_RETURN_ON_ERROR(TraitPin->GetSubPins()[5]->GetCPPType() == TEXT("int32"), TEXT("FAnimationAnimNextEditorTest_GraphAddTrait -> Unexpected trait pin type"));
-		UE_RETURN_ON_ERROR(TraitPin->GetSubPins()[5]->GetDefaultValue() == TEXT("3"), TEXT("FAnimationAnimNextEditorTest_GraphAddTrait -> Unexpected trait pin value"));
+		UE_RETURN_ON_ERROR(TraitPin->GetSubPins()[5]->GetDefaultValue() == TEXT("5"), TEXT("FAnimationAnimNextEditorTest_GraphAddTrait -> Unexpected trait pin value"));
 		UE_RETURN_ON_ERROR(TraitPin->GetSubPins()[5]->IsLazy(), TEXT("FAnimationAnimNextEditorTest_GraphAddTrait -> Expected lazy trait pin"));
 
 		// SomeOtherLatentInt32
 		UE_RETURN_ON_ERROR(TraitPin->GetSubPins()[6]->GetCPPType() == TEXT("int32"), TEXT("FAnimationAnimNextEditorTest_GraphAddTrait -> Unexpected trait pin type"));
-		UE_RETURN_ON_ERROR(TraitPin->GetSubPins()[6]->GetDefaultValue() == TEXT("3"), TEXT("FAnimationAnimNextEditorTest_GraphAddTrait -> Unexpected trait pin value"));
+		UE_RETURN_ON_ERROR(TraitPin->GetSubPins()[6]->GetDefaultValue() == TEXT("7"), TEXT("FAnimationAnimNextEditorTest_GraphAddTrait -> Unexpected trait pin value"));
 		UE_RETURN_ON_ERROR(TraitPin->GetSubPins()[6]->IsLazy(), TEXT("FAnimationAnimNextEditorTest_GraphAddTrait -> Expected lazy trait pin"));
 
-		// SomeLatentFloat
-		UE_RETURN_ON_ERROR(TraitPin->GetSubPins()[7]->GetCPPType() == TEXT("float"), TEXT("FAnimationAnimNextEditorTest_GraphAddTrait -> Unexpected trait pin type"));
-		UE_RETURN_ON_ERROR(TraitPin->GetSubPins()[7]->GetDefaultValue() == TEXT("34.000000"), TEXT("FAnimationAnimNextEditorTest_GraphAddTrait -> Unexpected trait pin value"));
+		// SomeOutOfOrderLatentBool
+		UE_RETURN_ON_ERROR(TraitPin->GetSubPins()[7]->GetCPPType() == TEXT("bool"), TEXT("FAnimationAnimNextEditorTest_GraphAddTrait -> Unexpected trait pin type"));
+		UE_RETURN_ON_ERROR(TraitPin->GetSubPins()[7]->GetDefaultValue() == TEXT("False"), TEXT("FAnimationAnimNextEditorTest_GraphAddTrait -> Unexpected trait pin value"));
 		UE_RETURN_ON_ERROR(TraitPin->GetSubPins()[7]->IsLazy(), TEXT("FAnimationAnimNextEditorTest_GraphAddTrait -> Expected lazy trait pin"));
+
+		// SomeLatentVector
+		UE_RETURN_ON_ERROR(TraitPin->GetSubPins()[8]->GetCPPType() == TEXT("FVector"), TEXT("FAnimationAnimNextEditorTest_GraphAddTrait -> Unexpected trait pin type"));
+		UE_RETURN_ON_ERROR(TraitPin->GetSubPins()[8]->GetDefaultValue() == TEXT("(X=1.000000,Y=1.000000,Z=1.000000)"), TEXT("FAnimationAnimNextEditorTest_GraphAddTrait -> Unexpected trait pin value"));
+		UE_RETURN_ON_ERROR(TraitPin->GetSubPins()[8]->IsLazy(), TEXT("FAnimationAnimNextEditorTest_GraphAddTrait -> Expected lazy trait pin"));
+
+		// SomeLatentFloat
+		UE_RETURN_ON_ERROR(TraitPin->GetSubPins()[9]->GetCPPType() == TEXT("float"), TEXT("FAnimationAnimNextEditorTest_GraphAddTrait -> Unexpected trait pin type"));
+		UE_RETURN_ON_ERROR(TraitPin->GetSubPins()[9]->GetDefaultValue() == TEXT("34.000000"), TEXT("FAnimationAnimNextEditorTest_GraphAddTrait -> Unexpected trait pin value"));
+		UE_RETURN_ON_ERROR(TraitPin->GetSubPins()[9]->IsLazy(), TEXT("FAnimationAnimNextEditorTest_GraphAddTrait -> Expected lazy trait pin"));
 	}
 
 	Tests::FUtils::CleanupAfterTests();
@@ -258,7 +271,7 @@ bool FAnimationAnimNextEditorTest_GraphTraitOperations::RunTest(const FString& I
 			UE_RETURN_ON_ERROR(TraitPin->GetCPPTypeObject() == FRigDecorator_AnimNextCppDecorator::StaticStruct(), TEXT("FAnimationAnimNextEditorTest_GraphTraitOperations -> Unexpected pin type"));
 
 			// Our first sub-pin is the hard coded script struct member that parametrizes the trait, dynamic trait sub-pins follow
-			UE_RETURN_ON_ERROR(TraitPin->GetSubPins().Num() == 8, TEXT("FAnimationAnimNextEditorTest_GraphAddTrait -> Unexpected trait sub pins"));
+			UE_RETURN_ON_ERROR(TraitPin->GetSubPins().Num() == 10, TEXT("FAnimationAnimNextEditorTest_GraphAddTrait -> Unexpected trait sub pins"));
 
 			// UpdateCount
 			UE_RETURN_ON_ERROR(TraitPin->GetSubPins()[1]->GetCPPType() == TEXT("int32"), TEXT("FAnimationAnimNextEditorTest_GraphAddTrait -> Unexpected trait pin type"));
@@ -282,18 +295,28 @@ bool FAnimationAnimNextEditorTest_GraphTraitOperations::RunTest(const FString& I
 
 			// SomeLatentInt32
 			UE_RETURN_ON_ERROR(TraitPin->GetSubPins()[5]->GetCPPType() == TEXT("int32"), TEXT("FAnimationAnimNextEditorTest_GraphAddTrait -> Unexpected trait pin type"));
-			UE_RETURN_ON_ERROR(TraitPin->GetSubPins()[5]->GetDefaultValue() == TEXT("3"), TEXT("FAnimationAnimNextEditorTest_GraphAddTrait -> Unexpected trait pin value"));
+			UE_RETURN_ON_ERROR(TraitPin->GetSubPins()[5]->GetDefaultValue() == TEXT("5"), TEXT("FAnimationAnimNextEditorTest_GraphAddTrait -> Unexpected trait pin value"));
 			UE_RETURN_ON_ERROR(TraitPin->GetSubPins()[5]->IsLazy(), TEXT("FAnimationAnimNextEditorTest_GraphAddTrait -> Expected lazy trait pin"));
 
 			// SomeOtherLatentInt32
 			UE_RETURN_ON_ERROR(TraitPin->GetSubPins()[6]->GetCPPType() == TEXT("int32"), TEXT("FAnimationAnimNextEditorTest_GraphAddTrait -> Unexpected trait pin type"));
-			UE_RETURN_ON_ERROR(TraitPin->GetSubPins()[6]->GetDefaultValue() == TEXT("3"), TEXT("FAnimationAnimNextEditorTest_GraphAddTrait -> Unexpected trait pin value"));
+			UE_RETURN_ON_ERROR(TraitPin->GetSubPins()[6]->GetDefaultValue() == TEXT("7"), TEXT("FAnimationAnimNextEditorTest_GraphAddTrait -> Unexpected trait pin value"));
 			UE_RETURN_ON_ERROR(TraitPin->GetSubPins()[6]->IsLazy(), TEXT("FAnimationAnimNextEditorTest_GraphAddTrait -> Expected lazy trait pin"));
 
-			// SomeLatentFloat
-			UE_RETURN_ON_ERROR(TraitPin->GetSubPins()[7]->GetCPPType() == TEXT("float"), TEXT("FAnimationAnimNextEditorTest_GraphAddTrait -> Unexpected trait pin type"));
-			UE_RETURN_ON_ERROR(TraitPin->GetSubPins()[7]->GetDefaultValue() == TEXT("34.000000"), TEXT("FAnimationAnimNextEditorTest_GraphAddTrait -> Unexpected trait pin value"));
+			// SomeOutOfOrderLatentBool
+			UE_RETURN_ON_ERROR(TraitPin->GetSubPins()[7]->GetCPPType() == TEXT("bool"), TEXT("FAnimationAnimNextEditorTest_GraphAddTrait -> Unexpected trait pin type"));
+			UE_RETURN_ON_ERROR(TraitPin->GetSubPins()[7]->GetDefaultValue() == TEXT("False"), TEXT("FAnimationAnimNextEditorTest_GraphAddTrait -> Unexpected trait pin value"));
 			UE_RETURN_ON_ERROR(TraitPin->GetSubPins()[7]->IsLazy(), TEXT("FAnimationAnimNextEditorTest_GraphAddTrait -> Expected lazy trait pin"));
+
+			// SomeLatentVector
+			UE_RETURN_ON_ERROR(TraitPin->GetSubPins()[8]->GetCPPType() == TEXT("FVector"), TEXT("FAnimationAnimNextEditorTest_GraphAddTrait -> Unexpected trait pin type"));
+			UE_RETURN_ON_ERROR(TraitPin->GetSubPins()[8]->GetDefaultValue() == TEXT("(X=1.000000,Y=1.000000,Z=1.000000)"), TEXT("FAnimationAnimNextEditorTest_GraphAddTrait -> Unexpected trait pin value"));
+			UE_RETURN_ON_ERROR(TraitPin->GetSubPins()[8]->IsLazy(), TEXT("FAnimationAnimNextEditorTest_GraphAddTrait -> Expected lazy trait pin"));
+
+			// SomeLatentFloat
+			UE_RETURN_ON_ERROR(TraitPin->GetSubPins()[9]->GetCPPType() == TEXT("float"), TEXT("FAnimationAnimNextEditorTest_GraphAddTrait -> Unexpected trait pin type"));
+			UE_RETURN_ON_ERROR(TraitPin->GetSubPins()[9]->GetDefaultValue() == TEXT("34.000000"), TEXT("FAnimationAnimNextEditorTest_GraphAddTrait -> Unexpected trait pin value"));
+			UE_RETURN_ON_ERROR(TraitPin->GetSubPins()[9]->IsLazy(), TEXT("FAnimationAnimNextEditorTest_GraphAddTrait -> Expected lazy trait pin"));
 		}
 
 		// --- Undo Add Trait ---
@@ -477,14 +500,16 @@ bool FAnimationAnimNextRuntimeTest_GraphExecute::RunTest(const FString& InParame
 			(void)UE::UAF::EvaluateGraph(EvaluateGraphContext);
 		}
 
-		AddErrorIfFalse(Messages.Num() == 7, "FAnimationAnimNextRuntimeTest_GraphExecute - > Unexpected message count");
+		AddErrorIfFalse(Messages.Num() == 9, "FAnimationAnimNextRuntimeTest_GraphExecute - > Unexpected message count");
 		AddErrorIfFalse(Messages[0] == "UpdateCount == 1", "FAnimationAnimNextRuntimeTest_GraphExecute -> Unexpected update count");
 		AddErrorIfFalse(Messages[1] == "SomeInt32 == 78", "FAnimationAnimNextRuntimeTest_GraphExecute -> Unexpected SomeInt32 value");
 		AddErrorIfFalse(Messages[2] == "SomeFloat == 142.33", "FAnimationAnimNextRuntimeTest_GraphExecute -> Unexpected SomeFloat value");
-		AddErrorIfFalse(Messages[3] == "SomeLatentInt32 == 3", "FAnimationAnimNextRuntimeTest_GraphExecute -> Unexpected SomeLatentInt32 value");
-		AddErrorIfFalse(Messages[4] == "SomeOtherLatentInt32 == 3", "FAnimationAnimNextRuntimeTest_GraphExecute -> Unexpected SomeOtherLatentInt32 value");
-		AddErrorIfFalse(Messages[5] == "SomeLatentFloat == 34.00", "FAnimationAnimNextRuntimeTest_GraphExecute -> Unexpected SomeLatentFloat value");
-		AddErrorIfFalse(Messages[6] == "EvaluateCount == 1", "FAnimationAnimNextRuntimeTest_GraphExecute -> Unexpected evaluate count");
+		AddErrorIfFalse(Messages[3] == "SomeLatentInt32 == 5", "FAnimationAnimNextRuntimeTest_GraphExecute -> Unexpected SomeLatentInt32 value");
+		AddErrorIfFalse(Messages[4] == "SomeOtherLatentInt32 == 7", "FAnimationAnimNextRuntimeTest_GraphExecute -> Unexpected SomeOtherLatentInt32 value");
+		AddErrorIfFalse(Messages[5] == "SomeOutOfOrderLatentBool == 0", "FAnimationAnimNextRuntimeTest_GraphExecute -> Unexpected SomeOutOfOrderLatentBool value");
+		AddErrorIfFalse(Messages[6] == "SomeLatentVector == [1.00, 1.00, 1.00]", "FAnimationAnimNextRuntimeTest_GraphExecute -> Unexpected SomeLatentVector value");
+		AddErrorIfFalse(Messages[7] == "SomeLatentFloat == 34.00", "FAnimationAnimNextRuntimeTest_GraphExecute -> Unexpected SomeLatentFloat value");
+		AddErrorIfFalse(Messages[8] == "EvaluateCount == 1", "FAnimationAnimNextRuntimeTest_GraphExecute -> Unexpected evaluate count");
 	}
 
 	Tests::FUtils::CleanupAfterTests();
@@ -510,6 +535,7 @@ bool FAnimationAnimNextRuntimeTest_GraphExecuteLatent::RunTest(const FString& In
 		UE_RETURN_ON_ERROR(EditorData != nullptr, "FAnimationAnimNextRuntimeTest_GraphExecuteLatent -> Failed to find module editor data");
 
 		EditorData->AddVariable(TEXT("TestIntVar"), FAnimNextParamType::GetType<int32>(), TEXT("34"));
+		EditorData->AddVariable(TEXT("TestVec3Var"), FAnimNextParamType::GetType<FVector>(), TEXT("(X=2.5,Y=1.25,Z=-7.75)"));
 		
 		UAnimNextController* Controller = Cast<UAnimNextController>(EditorData->GetRigVMClient()->GetController(EditorData->GetRigVMClient()->GetDefaultModel()));
 		UE_RETURN_ON_ERROR(Controller != nullptr, "FAnimationAnimNextRuntimeTest_GraphExecuteLatent -> Failed to get RigVM controller");
@@ -573,7 +599,7 @@ bool FAnimationAnimNextRuntimeTest_GraphExecuteLatent::RunTest(const FString& In
 
 		Controller->SetPinDefaultValue(DecoratorPin->GetSubPins()[3]->GetPinPath(), TEXT("78"));		// SomeInt32
 		Controller->SetPinDefaultValue(DecoratorPin->GetSubPins()[4]->GetPinPath(), TEXT("142.33"));	// SomeFloat
-		Controller->SetPinDefaultValue(DecoratorPin->GetSubPins()[7]->GetPinPath(), TEXT("1123.31"));	// SomeLatentFloat, inline value on latent pin
+		Controller->SetPinDefaultValue(DecoratorPin->GetSubPins()[9]->GetPinPath(), TEXT("1123.31"));	// SomeLatentFloat, inline value on latent pin
 
 		// Set some latent values on our trait
 		{
@@ -596,6 +622,13 @@ bool FAnimationAnimNextRuntimeTest_GraphExecuteLatent::RunTest(const FString& In
 			Controller->AddLink(
 				GetVariableNode->FindPin(TEXT("Value")),
 				DecoratorPin->GetSubPins()[6]);	// SomeOtherLatentInt32
+
+			URigVMVariableNode* GetVariableNode1 = Controller->AddVariableNode(TEXT("TestVec3Var"), TEXT("FVector"), nullptr, true, TEXT(""));
+			UE_RETURN_ON_ERROR(GetVariableNode1 != nullptr, TEXT("FAnimationAnimNextRuntimeTest_GraphExecuteLatent -> Failed to create variable node"));
+
+			Controller->AddLink(
+				GetVariableNode1->FindPin(TEXT("Value")),
+				DecoratorPin->GetSubPins()[8]);	// SomeOtherLatentInt32
 		}
 
 		TSharedPtr<FAnimNextGraphInstance> GraphInstance = AnimationGraph->AllocateInstance();
@@ -615,14 +648,16 @@ bool FAnimationAnimNextRuntimeTest_GraphExecuteLatent::RunTest(const FString& In
 			(void)UE::UAF::EvaluateGraph(EvaluateGraphContext);
 		}
 
-		AddErrorIfFalse(Messages.Num() == 7, "FAnimationAnimNextRuntimeTest_GraphExecuteLatent - > Unexpected message count");
+		AddErrorIfFalse(Messages.Num() == 9, "FAnimationAnimNextRuntimeTest_GraphExecuteLatent - > Unexpected message count");
 		AddErrorIfFalse(Messages[0] == "UpdateCount == 1", "FAnimationAnimNextRuntimeTest_GraphExecuteLatent -> Unexpected update count");
 		AddErrorIfFalse(Messages[1] == "SomeInt32 == 78", "FAnimationAnimNextRuntimeTest_GraphExecuteLatent -> Unexpected SomeInt32 value");
 		AddErrorIfFalse(Messages[2] == "SomeFloat == 142.33", "FAnimationAnimNextRuntimeTest_GraphExecuteLatent -> Unexpected SomeFloat value");
 		AddErrorIfFalse(Messages[3] == "SomeLatentInt32 == 33", "FAnimationAnimNextRuntimeTest_GraphExecuteLatent -> Unexpected SomeLatentInt32 value");
 		AddErrorIfFalse(Messages[4] == "SomeOtherLatentInt32 == 34", "FAnimationAnimNextRuntimeTest_GraphExecuteLatent -> Unexpected SomeOtherLatentInt32 value");
-		AddErrorIfFalse(Messages[5] == "SomeLatentFloat == 1123.31", "FAnimationAnimNextRuntimeTest_GraphExecuteLatent -> Unexpected SomeLatentFloat value");
-		AddErrorIfFalse(Messages[6] == "EvaluateCount == 1", "FAnimationAnimNextRuntimeTest_GraphExecuteLatent -> Unexpected evaluate count");
+		AddErrorIfFalse(Messages[5] == "SomeOutOfOrderLatentBool == 0", "FAnimationAnimNextRuntimeTest_GraphExecuteLatent -> Unexpected SomeOutOfOrderLatentBool value");
+		AddErrorIfFalse(Messages[6] == "SomeLatentVector == [2.50, 1.25, -7.75]", "FAnimationAnimNextRuntimeTest_GraphExecuteLatent -> Unexpected SomeLatentVector value");
+		AddErrorIfFalse(Messages[7] == "SomeLatentFloat == 1123.31", "FAnimationAnimNextRuntimeTest_GraphExecuteLatent -> Unexpected SomeLatentFloat value");
+		AddErrorIfFalse(Messages[8] == "EvaluateCount == 1", "FAnimationAnimNextRuntimeTest_GraphExecuteLatent -> Unexpected evaluate count");
 	}
 
 	Tests::FUtils::CleanupAfterTests();
@@ -630,7 +665,6 @@ bool FAnimationAnimNextRuntimeTest_GraphExecuteLatent::RunTest(const FString& In
 	return true;
 }
 
-
 IMPLEMENT_SIMPLE_AUTOMATION_TEST(FAnimationAnimNextRuntimeTest_Variables, "Animation.AnimNext.Runtime.Graph.Variables", EAutomationTestFlags::EditorContext | EAutomationTestFlags::EngineFilter)
 
 bool FAnimationAnimNextRuntimeTest_Variables::RunTest(const FString& InParameters)
diff --git a/Engine/Plugins/Experimental/UAF/UAFAnimGraph/Source/UAFAnimGraphTestSuite/Private/AnimNextAnimGraphTraitGraphTest.h b/Engine/Plugins/Experimental/UAF/UAFAnimGraph/Source/UAFAnimGraphTestSuite/Private/AnimNextAnimGraphTraitGraphTest.h
--- a/Engine/Plugins/Experimental/UAF/UAFAnimGraph/Source/UAFAnimGraphTestSuite/Private/AnimNextAnimGraphTraitGraphTest.h
+++ b/Engine/Plugins/Experimental/UAF/UAFAnimGraph/Source/UAFAnimGraphTestSuite/Private/AnimNextAnimGraphTraitGraphTest.h
@@ -28,17 +28,26 @@ struct FTestTraitSharedData : public FAnimNextTraitSharedData
 	float SomeFloat = 34.0f;
 
 	UPROPERTY(meta = (Input))
-	int32 SomeLatentInt32 = 3;			// MathAdd with constants, latent
+	int32 SomeLatentInt32 = 5;			// MathAdd with constants, latent
 
 	UPROPERTY(meta = (Input))
-	int32 SomeOtherLatentInt32 = 3;		// GetParameter, latent
+	int32 SomeOtherLatentInt32 = 7;		// GetParameter, latent
+
+	UPROPERTY(meta = (Input))
+	bool SomeOutOfOrderLatentBool = false;	// GetParameter, latent
+
+	UPROPERTY(meta = (Input))
+	FVector SomeLatentVector = FVector::OneVector;	// GetParameter, latent
 
 	UPROPERTY(meta = (Input))
 	float SomeLatentFloat = 34.0f;		// Inline value, not latent
 
 	#define TRAIT_LATENT_PROPERTIES_ENUMERATOR(GeneratorMacro) \
 		GeneratorMacro(SomeLatentInt32) \
+		/* Out of order with property definition above */ \
+		GeneratorMacro(SomeOutOfOrderLatentBool) \
 		GeneratorMacro(SomeOtherLatentInt32) \
+		GeneratorMacro(SomeLatentVector) \
 		GeneratorMacro(SomeLatentFloat) \
 
 	GENERATE_TRAIT_LATENT_PROPERTIES(FTestTraitSharedData, TRAIT_LATENT_PROPERTIES_ENUMERATOR)
@@ -56,4 +65,4 @@ struct FTestDerivedVector : public FVector
 
 	UPROPERTY()
 	int32 W = 1;
-};
\ No newline at end of file
+};
diff --git a/Engine/Plugins/Experimental/UAF/UAFControlRig/Source/UAFControlRig/Internal/ControlRigTrait.h b/Engine/Plugins/Experimental/UAF/UAFControlRig/Source/UAFControlRig/Internal/ControlRigTrait.h
--- a/Engine/Plugins/Experimental/UAF/UAFControlRig/Source/UAFControlRig/Internal/ControlRigTrait.h
+++ b/Engine/Plugins/Experimental/UAF/UAFControlRig/Source/UAFControlRig/Internal/ControlRigTrait.h
@@ -274,7 +274,7 @@ protected:
 		const FString& InDefaultValue,
 		struct FRigVMPinInfoArray& OutPinArray) const override;
 
-	bool IsPropertyLatent(const FAnimNextTraitSharedData& InSharedData, FName PropertyName) const override;
+	virtual int32 GetLatentPropertyIndex(const FAnimNextTraitSharedData& InSharedData, FName PropertyName) const override;
 
 	virtual uint32 GetLatentPropertyHandles(
 		const FAnimNextTraitSharedData* InSharedData,
diff --git a/Engine/Plugins/Experimental/UAF/UAFControlRig/Source/UAFControlRig/Private/ControlRigTrait.cpp b/Engine/Plugins/Experimental/UAF/UAFControlRig/Source/UAFControlRig/Private/ControlRigTrait.cpp
--- a/Engine/Plugins/Experimental/UAF/UAFControlRig/Source/UAFControlRig/Private/ControlRigTrait.cpp
+++ b/Engine/Plugins/Experimental/UAF/UAFControlRig/Source/UAFControlRig/Private/ControlRigTrait.cpp
@@ -475,14 +475,14 @@ FTraitLatentPropertyMemoryLayout FControlRigTrait::GetLatentPropertyMemoryLayout
 	if (LatentPropertyIndex == INDEX_NONE)
 	{
 		// This property isn't being tracked, ignore it
-		return FTraitLatentPropertyMemoryLayout{ 0, 1 };
+		return FTraitLatentPropertyMemoryLayout();
 	}
 
 	const uint32 PropertyLayout = ControlRigSharedData.LatentPropertyMemoryLayouts[LatentPropertyIndex];
 	const uint32 PropertySize = PropertyLayout >> 16;
 	const uint32 PropertyAlignment = PropertyLayout & MAX_uint16;
 
-	return FTraitLatentPropertyMemoryLayout{ PropertySize, PropertyAlignment };
+	return FTraitLatentPropertyMemoryLayout{ PropertySize, PropertyAlignment, LatentPropertyIndex };
 }
 
 // --- IUpdate impl --- 
@@ -743,29 +743,10 @@ void FControlRigTrait::GetProgrammaticPins(FAnimNextTraitSharedData* InSharedDat
 	}
 }
 
-bool FControlRigTrait::IsPropertyLatent(const FAnimNextTraitSharedData& InSharedData, FName PropertyName) const
+int32 FControlRigTrait::GetLatentPropertyIndex(const FAnimNextTraitSharedData& InSharedData, FName PropertyName) const
 {
 	const FControlRigTraitSharedData& SharedData = (const FControlRigTraitSharedData&)InSharedData;
-
-	// --- Exposed Public Variables ---
-	if (SharedData.ExposedPropertyVariableNames.Num() > 0)
-	{
-		if (SharedData.ExposedPropertyVariableNames.Contains(PropertyName))
-		{
-			return true;
-		}
-	}
-
-	// --- Exposed Controls ---
-	if (SharedData.ExposedPropertyControlNames.Num() > 0)
-	{
-		if (SharedData.ExposedPropertyControlNames.Contains(PropertyName))
-		{
-			return true;
-		}
-	}
-
-	return false;
+	return SharedData.LatentProperties.IndexOfByKey(PropertyName);
 }
 
 uint32 FControlRigTrait::GetLatentPropertyHandles(
