UAF: Anim graph node editing workflow improvements

Trait stacks can now be edited in the details panel. Removed/hid other methods of editing trait stacks.
Trait stacks are now validated on graph compilation.
Fixed up some RigVM layout editing bugs:
- Fixed pins getting reversed in categories when undoing a layout edit
- Fixed error popups when undoing a layout edit
Fixed up comments for trait shared data structs, as this is where tooltips are drawn from
Fixed issues that were preventing multple traits of identical type from existing in one stack (e.g. multiple call function traits)
Fixed function picker to not always refresh - this was causing performance issues when creating stacks with multiple call function traits
Fixed BlendSpace player template being spawned by default with a warning

diff --git a/Engine/Plugins/Experimental/UAF/UAF/Source/UAFEditor/Internal/ITraitStackEditor.h b/Engine/Plugins/Experimental/UAF/UAF/Source/UAFEditor/Internal/ITraitStackEditor.h
--- a/Engine/Plugins/Experimental/UAF/UAF/Source/UAFEditor/Internal/ITraitStackEditor.h
+++ b/Engine/Plugins/Experimental/UAF/UAF/Source/UAFEditor/Internal/ITraitStackEditor.h
@@ -7,6 +7,10 @@
 #include "UObject/WeakObjectPtr.h"
 
 class UAnimNextEdGraphNode;
+class SWidget;
+class IDetailPropertyRow;
+class IDetailCategoryBuilder;
+class URigVMController;
 
 namespace UE::Workspace
 {
@@ -37,6 +41,12 @@ public:
 
 	// Sets the trait data to be displayed for the specified workspace editor instance
 	virtual void SetTraitData(const TSharedRef<UE::Workspace::IWorkspaceEditor>& InWorkspaceEditor, const FTraitStackData& InTraitStackData) = 0;
+
+	// Create a header widget for the trait in a trait stack node corresponding to InPropertyRow
+	virtual TSharedRef<SWidget> CreateTraitHeaderWidget(IDetailCategoryBuilder& InCategory, IDetailPropertyRow& InPropertyRow, URigVMController* InController) = 0;
+
+	// Create a header widget for the stack in a trait stack node 
+	virtual TSharedRef<SWidget> CreateTraitStackHeaderWidget(IDetailCategoryBuilder& InCategory, URigVMController* InController) = 0;
 };
 
 }
\ No newline at end of file
diff --git a/Engine/Plugins/Experimental/UAF/UAF/Source/UAFEditor/Private/AnimNextEditorModule.cpp b/Engine/Plugins/Experimental/UAF/UAF/Source/UAFEditor/Private/AnimNextEditorModule.cpp
--- a/Engine/Plugins/Experimental/UAF/UAF/Source/UAFEditor/Private/AnimNextEditorModule.cpp
+++ b/Engine/Plugins/Experimental/UAF/UAF/Source/UAFEditor/Private/AnimNextEditorModule.cpp
@@ -987,6 +987,7 @@ void FAnimNextEditorModule::RegisterWorkspaceDocumentTypes(Workspace::IWorkspace
 		}
 		Controller->RemoveNodes(NodesToRemove, true, true);
 		Controller->CloseUndoBracket();
+		InContext.WorkspaceEditor->SetDetailsObjects({});
 	});
 	GraphArgs.OnCanCopySelectedNodes = Workspace::FOnCanPerformActionOnSelectedNodes::CreateLambda([this](const Workspace::FWorkspaceEditorContext& InContext, const FGraphPanelSelectionSet& InSelectedNodes)
 	{
@@ -1094,6 +1095,8 @@ void FAnimNextEditorModule::RegisterWorkspaceDocumentTypes(Workspace::IWorkspace
 				Controller->OpenUndoBracket(TEXT("Cut Nodes."));
 				OnDeleteSelectedNodes.Execute(InContext, InSelectedNodes);
 				Controller->CloseUndoBracket();
+
+				InContext.WorkspaceEditor->SetDetailsObjects({});
 			}
 		});
 	UE::Workspace::FOnCanPasteNodes& OnCanPasteNodes = GraphArgs.OnCanPasteNodes;
diff --git a/Engine/Plugins/Experimental/UAF/UAF/Source/UAFEditor/Private/Common/AnimNextEdGraphNodeCustomization.cpp b/Engine/Plugins/Experimental/UAF/UAF/Source/UAFEditor/Private/Common/AnimNextEdGraphNodeCustomization.cpp
--- a/Engine/Plugins/Experimental/UAF/UAF/Source/UAFEditor/Private/Common/AnimNextEdGraphNodeCustomization.cpp
+++ b/Engine/Plugins/Experimental/UAF/UAF/Source/UAFEditor/Private/Common/AnimNextEdGraphNodeCustomization.cpp
@@ -45,8 +45,10 @@ void FAnimNextEdGraphNodeCustomization::PendingDelete()
 
 void FAnimNextEdGraphNodeCustomization::CustomizeDetails(IDetailLayoutBuilder& DetailBuilder)
 {
-	TArray<TWeakObjectPtr<UObject>> Objects;
-	DetailBuilder.GetObjectsBeingCustomized(Objects);
+	TArray<TWeakObjectPtr<UObject>> WeakObjects;
+	DetailBuilder.GetObjectsBeingCustomized(WeakObjects);
+	TArray<UObject*> Objects;
+	Algo::TransformIf(WeakObjects, Objects, [](TWeakObjectPtr<UObject> InObject){ return InObject.Get() != nullptr; }, [](TWeakObjectPtr<UObject> InObject){ return InObject.Get(); });
 
 	RigVMGraphDetailCustomizationImpl.Reset();
 
@@ -58,12 +60,12 @@ void FAnimNextEdGraphNodeCustomization::CustomizeDetails(IDetailLayoutBuilder& D
 	CustomizeObjects(DetailBuilder, Objects);
 }
 
-void FAnimNextEdGraphNodeCustomization::CustomizeObjects(IDetailLayoutBuilder& DetailBuilder, const TArray<TWeakObjectPtr<UObject>>& Objects)
+void FAnimNextEdGraphNodeCustomization::CustomizeObjects(IDetailLayoutBuilder& DetailBuilder, const TArray<UObject*>& Objects)
 {
 	const int32 NumNodes = Objects.Num();
 	for (int32 NodeIndex = 0; NodeIndex < NumNodes; ++NodeIndex)
 	{
-		UAnimNextEdGraphNode* EdGraphNode = Cast<UAnimNextEdGraphNode>(Objects[NodeIndex].Get());
+		UAnimNextEdGraphNode* EdGraphNode = Cast<UAnimNextEdGraphNode>(Objects[NodeIndex]);
 		if(EdGraphNode == nullptr)
 		{
 			continue;
@@ -126,15 +128,42 @@ void FAnimNextEdGraphNodeCustomization::CustomizeObjects(IDetailLayoutBuilder& D
 		}
 	}
 
+	// Create a header widget to edit the trait stack
+	bool bAreAllNodesTraitStacks = true;
+	URigVMController* Controller = nullptr;
+	for (TWeakObjectPtr<UObject> WeakObject : Objects)
+	{
+		if (UAnimNextEdGraphNode* EdGraphNode = Cast<UAnimNextEdGraphNode>(WeakObject.Get()))
+		{
+			if (EdGraphNode->GetModelNode())
+			{
+				bAreAllNodesTraitStacks &= EdGraphNode->GetModelNode()->GetClass()->GetClassPathName() == FTopLevelAssetPath(TEXT("/Script/UAFAnimGraphUncookedOnly"), TEXT("AnimNextTraitStackUnitNode"));
+				Controller = EdGraphNode->GetController();
+			}
+		}
+	}
+
+	if(bAreAllNodesTraitStacks && Controller && IModularFeatures::Get().IsModularFeatureAvailable(ITraitStackEditor::ModularFeatureName))
+	{
+		IDetailCategoryBuilder& HeaderCategory = DetailBuilder.EditCategory(TEXT("TraitStack"), FText::GetEmpty(), ECategoryPriority::Important);
+		HeaderCategory.AddCustomRow(FText::GetEmpty());
+		HeaderCategory.SetIsEmpty(true);
+
+		ITraitStackEditor& TraitStackEditor = IModularFeatures::Get().GetModularFeature<ITraitStackEditor>(ITraitStackEditor::ModularFeatureName);
+		HeaderCategory.HeaderContent(
+			TraitStackEditor.CreateTraitStackHeaderWidget(HeaderCategory, Controller),
+			true);
+	}
+	
 	for (TSharedPtr<FCategoryDetailsData>& DetailsData : CategoryDetailsData)
 	{
 		PopulateCategory(DetailBuilder, DetailsData);
 	}
-	
+
 	// Pass the TraitStack EdGraphNode to the TraitEditor (or nullptr if we select other node type, like RigVM)
-	if (Objects.Num() == 1)
+	if (Objects.Num() == 1 && bAreAllNodesTraitStacks)
 	{
-		if (UAnimNextEdGraphNode* EdGraphNode = Cast<UAnimNextEdGraphNode>(Objects[0].Get()))
+		if (UAnimNextEdGraphNode* EdGraphNode = Cast<UAnimNextEdGraphNode>(Objects[0]))
 		{
 			WeakModelNode = EdGraphNode->GetModelNode();
 			if (WeakModelNode.Get())
@@ -169,8 +198,26 @@ void FAnimNextEdGraphNodeCustomization::CustomizeObjects(IDetailLayoutBuilder& D
 						}
 					}));
 
-					FDetailWidgetRow& Row = CustomizationCategory.AddCustomRow(LOCTEXT("LayoutSearchString", "Node Layout"));
+					// Widget that prevents RigVM compilation as long as it is alive (used to prevent popup from closing for each layout operation)
+					class SRigVMCompileLock : public SBox
+					{
+						SLATE_BEGIN_ARGS(SRigVMCompileLock) {}
+						SLATE_DEFAULT_SLOT(FArguments, Content)
+						SLATE_END_ARGS()
+
+						void Construct(const FArguments& InArgs, URigVMController* InController)
+						{
+							Scope = MakeShared<FRigVMControllerCompileBracketScope>(InController);
+							SBox::Construct(SBox::FArguments()
+								.WidthOverride(300.0f)
+								.HeightOverride(300.0f)
+								.Content()[InArgs._Content.Widget]);
+						}
 
+						TSharedPtr<FRigVMControllerCompileBracketScope> Scope;
+					};
+					
+					FDetailWidgetRow& Row = CustomizationCategory.AddCustomRow(LOCTEXT("LayoutSearchString", "Node Layout"));
 					Row
 					.NameContent()
 					[
@@ -192,12 +239,16 @@ void FAnimNextEdGraphNodeCustomization::CustomizeObjects(IDetailLayoutBuilder& D
 								.Font(IDetailLayoutBuilder::GetDetailFont())
 							]
 							.HasDownArrow(true)
-							.OnGetMenuContent_Lambda([this]()
+							.OnGetMenuContent_Lambda([this, WeakEdGraphNode = TWeakObjectPtr<UAnimNextEdGraphNode>(EdGraphNode)]() -> TSharedRef<SWidget>
 							{
+								UAnimNextEdGraphNode* EdGraphNode = WeakEdGraphNode.Get();
+								if (EdGraphNode == nullptr)
+								{
+									return SNullWidget::NullWidget;
+								}
+								
 								return
-									SNew(SBox)
-									.WidthOverride(300.0f)
-									.HeightOverride(300.0f)
+									SNew(SRigVMCompileLock, EdGraphNode->GetController())
 									[
 										SNew(SRigVMNodeLayoutWidget)
 										.OnGetUncategorizedPins(this, &FAnimNextEdGraphNodeCustomization::GetUncategorizedPins)
@@ -233,6 +284,13 @@ void FAnimNextEdGraphNodeCustomization::CustomizeObjects(IDetailLayoutBuilder& D
 			}
 		}
 	}
+
+	// hide color editor fo non-trait stacks
+	if (!bAreAllNodesTraitStacks)
+	{
+		DetailBuilder.GetProperty(GET_MEMBER_NAME_CHECKED(UAnimNextEdGraphNode, NodeColor))
+		->MarkHiddenByCustomization();
+	}
 }
 
 void FAnimNextEdGraphNodeCustomization::GenerateTraitData(UAnimNextEdGraphNode* EdGraphNode, TArray<TSharedPtr<FCategoryDetailsData>>& CategoryDetailsData)
@@ -243,6 +301,8 @@ void FAnimNextEdGraphNodeCustomization::GenerateTraitData(UAnimNextEdGraphNode*
 		const TArray<URigVMPin*> TraitPins = ModelNode->GetTraitPins();
 		if (TraitPins.Num() > 0)
 		{
+			TMap<const UScriptStruct*, int32> TypeCounts; 
+
 			// For each Trait (represented as a pin in the node)
 			for (const URigVMPin* TraitPin : TraitPins)
 			{
@@ -270,11 +330,13 @@ void FAnimNextEdGraphNodeCustomization::GenerateTraitData(UAnimNextEdGraphNode*
 
 						if (TraitDataPtr == nullptr)
 						{
-							TraitDataPtr = &CategoryDetailsData.Add_GetRef(MakeShared<FTraitStackDetailsData>(TraitPin->GetFName()));
+							TraitDataPtr = &CategoryDetailsData.Add_GetRef(MakeShared<FTraitStackDetailsData>(TraitPin->GetFName(), TypeCounts.FindRef(TraitSharedInstanceData)));
 						}
 
 						TraitData = static_cast<FTraitStackDetailsData*>(TraitDataPtr->Get());
 
+						TypeCounts.FindOrAdd(TraitSharedInstanceData)++;
+
 						// Store EdGraphNode and ScopedTraitData, as we will need that later to transfer data if the user makes modifications to the Traits in the details panel
 						TraitData->EdGraphNodes.Add(EdGraphNode);
 						const int32 ScopedDataIndex = TraitData->ScopedSharedDataInstances.Add(MakeShared<FStructOnScope>(TraitSharedInstanceData));
@@ -385,9 +447,9 @@ void FAnimNextEdGraphNodeCustomization::PopulateCategory(IDetailLayoutBuilder& D
 
 	// Create a category with the DisplayName of the Trait Shared Data
 	const FString TraitDisplayName = *TraitData->ScopedSharedDataInstances[0]->GetStruct()->GetDisplayNameText().ToString();
-	const FName CategoryName = (TraitData->EdGraphNodes.Num() == 1)
+	const FName CategoryName = (TraitData->InstanceIndex == 0)
 		? *TraitDisplayName
-		: *FString::Printf(TEXT("%s (%d)"), *TraitDisplayName, TraitData->EdGraphNodes.Num());
+		: *FString::Printf(TEXT("%s %d"), *TraitDisplayName, TraitData->InstanceIndex);
 
 	IDetailCategoryBuilder& ParameterCategory = DetailBuilder.EditCategory(CategoryName, FText::GetEmpty(), ECategoryPriority::Important);
 	bool bHasProperties = false;
@@ -403,6 +465,16 @@ void FAnimNextEdGraphNodeCustomization::PopulateCategory(IDetailLayoutBuilder& D
 		return;
 	}
 
+	if(IModularFeatures::Get().IsModularFeatureAvailable(ITraitStackEditor::ModularFeatureName))
+	{
+		URigVMController* Controller = TraitData->EdGraphNodes[0]->GetController();
+		
+		ITraitStackEditor& TraitStackEditor = IModularFeatures::Get().GetModularFeature<ITraitStackEditor>(ITraitStackEditor::ModularFeatureName);
+		ParameterCategory.HeaderContent(
+			TraitStackEditor.CreateTraitHeaderWidget(ParameterCategory, *DetailPropertyRow, Controller),
+			true);
+	}
+
 	TraitData->RootPropertyHandle = PropertyHandle;
 
 	auto VisitChildren = [&bHasProperties](const TSharedRef<IPropertyHandle>& InPropertyHandle, TFunctionRef<void(const TSharedRef<IPropertyHandle>&)> InPredicate)
@@ -624,26 +696,38 @@ TArray<FString> FAnimNextEdGraphNodeCustomization::GetUncategorizedPins() const
 {
 	if(const URigVMNode* Node = GetNodeForLayout())
 	{
-		TArray<FString> PinPaths;
-		const TArray<URigVMPin*> AllPins = Node->GetAllPinsRecursively();
-		PinPaths.Reserve(AllPins.Num());
-		for(const URigVMPin* Pin : AllPins)
+		struct Local
 		{
-			if(Pin->IsExecuteContext())
-			{
-				continue;
-			}
-			if(Pin->GetDirection() != ERigVMPinDirection::Input &&
-				Pin->GetDirection() != ERigVMPinDirection::Visible)
+			static void VisitPinRecursively(URigVMPin* InPin, TArray<FString>& OutPinPaths)
 			{
-				continue;
-			}
-			if(!Pin->GetCategory().IsEmpty())
-			{
-				continue;
+				if(InPin->IsExecuteContext())
+				{
+					return;
+				}
+				if(InPin->GetDirection() != ERigVMPinDirection::Input &&
+					InPin->GetDirection() != ERigVMPinDirection::Visible)
+				{
+					return;
+				}
+				if(!InPin->GetCategory().IsEmpty())
+				{
+					return;
+				}
+
+				OutPinPaths.Add(InPin->GetSegmentPath(true));
+				for (URigVMPin* SubPin : InPin->GetSubPins())
+				{
+				    VisitPinRecursively(SubPin, OutPinPaths);
+				}
 			}
-			PinPaths.Add(Pin->GetSegmentPath(true));
+		};
+
+		TArray<FString> PinPaths;
+		for (URigVMPin* Pin : Node->GetPins())
+		{
+			Local::VisitPinRecursively(Pin, PinPaths);
 		}
+
 		return PinPaths;
 	}
 	return TArray<FString>();
diff --git a/Engine/Plugins/Experimental/UAF/UAF/Source/UAFEditor/Private/Common/AnimNextEdGraphNodeCustomization.h b/Engine/Plugins/Experimental/UAF/UAF/Source/UAFEditor/Private/Common/AnimNextEdGraphNodeCustomization.h
--- a/Engine/Plugins/Experimental/UAF/UAF/Source/UAFEditor/Private/Common/AnimNextEdGraphNodeCustomization.h
+++ b/Engine/Plugins/Experimental/UAF/UAF/Source/UAFEditor/Private/Common/AnimNextEdGraphNodeCustomization.h
@@ -74,14 +74,17 @@ protected:
 			: FCategoryDetailsData(FCategoryDetailsData::EType::TraitStack)
 		{
 		}
-		explicit FTraitStackDetailsData(const FName InName)
+		explicit FTraitStackDetailsData(const FName InName, int32 InInstanceIndex)
 			: FCategoryDetailsData(FCategoryDetailsData::EType::TraitStack, InName)
+			, InstanceIndex(InInstanceIndex)
 		{
 		}
 
 		TArray<TSharedPtr<FStructOnScope>> ScopedSharedDataInstances;
 		TSharedPtr<IPropertyHandle> RootPropertyHandle;
 		TArray<TSharedRef<IPropertyHandle>> PropertyHandles;
+		// For traits that appear more than once, this allows us to disambiguate
+		int32 InstanceIndex;
 	};
 
 	struct FRigVMNodeDetailsData : FCategoryDetailsData
@@ -100,7 +103,7 @@ protected:
 		TArray <TSharedPtr<FRigVMMemoryStorageStruct>> MemoryStorages;
 	};
 
-	void CustomizeObjects(IDetailLayoutBuilder& DetailBuilder, const TArray<TWeakObjectPtr<UObject>>& InObjects);
+	void CustomizeObjects(IDetailLayoutBuilder& DetailBuilder, const TArray<UObject*>& InObjects);
 
 	static void GenerateTraitData(UAnimNextEdGraphNode* EdGraphNode, TArray<TSharedPtr<FCategoryDetailsData>>& CategoryDetailsData);
 	static void GenerateRigVMData(UAnimNextEdGraphNode* EdGraphNode, TArray<TSharedPtr<FCategoryDetailsData>>& CategoryDetailsData);
diff --git a/Engine/Plugins/Experimental/UAF/UAF/Source/UAFEditor/Private/Common/SRigVMFunctionPicker.cpp b/Engine/Plugins/Experimental/UAF/UAF/Source/UAFEditor/Private/Common/SRigVMFunctionPicker.cpp
--- a/Engine/Plugins/Experimental/UAF/UAF/Source/UAFEditor/Private/Common/SRigVMFunctionPicker.cpp
+++ b/Engine/Plugins/Experimental/UAF/UAF/Source/UAFEditor/Private/Common/SRigVMFunctionPicker.cpp
@@ -201,8 +201,6 @@ void SRigVMFunctionPicker::Construct(const FArguments& InArgs)
 			]
 		]
 	];
-
-	RequestRefreshEntries();
 }
 
 void SRigVMFunctionPicker::RequestRefreshEntries()
diff --git a/Engine/Plugins/Experimental/UAF/UAF/Source/UAFUncookedOnly/Private/AnimNextRigVMAssetEditorData.cpp b/Engine/Plugins/Experimental/UAF/UAF/Source/UAFUncookedOnly/Private/AnimNextRigVMAssetEditorData.cpp
--- a/Engine/Plugins/Experimental/UAF/UAF/Source/UAFUncookedOnly/Private/AnimNextRigVMAssetEditorData.cpp
+++ b/Engine/Plugins/Experimental/UAF/UAF/Source/UAFUncookedOnly/Private/AnimNextRigVMAssetEditorData.cpp
@@ -1425,6 +1425,12 @@ void UAnimNextRigVMAssetEditorData::HandleModifiedEvent(ERigVMGraphNotifType InN
 			RequestAutoVMRecompilation(); // can not check if it is a trait pin, as it has been already removed
 			break;
 		}
+	case ERigVMGraphNotifType::PinCategoryChanged:
+	case ERigVMGraphNotifType::PinCategoriesChanged:
+		{
+			RequestAutoVMRecompilation();
+			break;
+		}
 	}
 	
 	// if the notification still has to be sent...
diff --git a/Engine/Plugins/Experimental/UAF/UAFAnimGraph/Source/UAFAnimGraph/Internal/Traits/BlendByBool.h b/Engine/Plugins/Experimental/UAF/UAFAnimGraph/Source/UAFAnimGraph/Internal/Traits/BlendByBool.h
--- a/Engine/Plugins/Experimental/UAF/UAFAnimGraph/Source/UAFAnimGraph/Internal/Traits/BlendByBool.h
+++ b/Engine/Plugins/Experimental/UAF/UAFAnimGraph/Source/UAFAnimGraph/Internal/Traits/BlendByBool.h
@@ -11,7 +11,8 @@
 
 #include "BlendByBool.generated.h"
 
-USTRUCT(meta = (DisplayName = "Blend By Bool"))
+/** A trait that can blend two discrete inputs through a boolean. */
+USTRUCT(meta = (DisplayName = "Blend By Bool", ShowTooltip=true))
 struct FAnimNextBlendByBoolTraitSharedData : public FAnimNextTraitSharedData
 {
 	GENERATED_BODY()
diff --git a/Engine/Plugins/Experimental/UAF/UAFAnimGraph/Source/UAFAnimGraph/Internal/Traits/BlendStackTrait.h b/Engine/Plugins/Experimental/UAF/UAFAnimGraph/Source/UAFAnimGraph/Internal/Traits/BlendStackTrait.h
--- a/Engine/Plugins/Experimental/UAF/UAFAnimGraph/Source/UAFAnimGraph/Internal/Traits/BlendStackTrait.h
+++ b/Engine/Plugins/Experimental/UAF/UAFAnimGraph/Source/UAFAnimGraph/Internal/Traits/BlendStackTrait.h
@@ -18,7 +18,8 @@
 
 #include "BlendStackTrait.generated.h"
 
-USTRUCT(meta = (DisplayName = "Blend Stack Core"))
+/** Used as base trait in state machine-like setups where state machine is the additive trait that pushes graphs and blend settings to the blend stack. */
+USTRUCT(meta = (DisplayName = "Blend Stack Core", ShowTooltip=true))
 struct FAnimNextBlendStackCoreTraitSharedData : public FAnimNextTraitSharedData
 {
 	GENERATED_BODY()
@@ -32,7 +33,11 @@ namespace UE::UAF
 	using FBlendStackCoreTraitData = FAnimNextBlendStackCoreTraitSharedData; 
 }
 
-USTRUCT(meta = (DisplayName = "Blend Stack"))
+/**
+ * Standalone blend stack that exposes desired object and blend time as pins.
+ * A new graph will be pushed every time the newest object does not match the pin (or bForceBlend is true).
+ */	
+USTRUCT(meta = (DisplayName = "Blend Stack", ShowTooltip=true))
 struct FAnimNextBlendStackTraitSharedData : public FAnimNextBlendStackCoreTraitSharedData
 {
 	GENERATED_BODY()
@@ -59,7 +64,11 @@ struct FAnimNextBlendStackTraitSharedData : public FAnimNextBlendStackCoreTraitS
 	#undef TRAIT_LATENT_PROPERTIES_ENUMERATOR
 };
 
-USTRUCT(meta = (DisplayName = "Blend Stack Requester"))
+/**
+ * Example additive trait used to push graphs to a Blend Stack base trait.
+ * Same behavior as FBlendStackTrait, but as additive (i.e. FBlendStackCoreTrait base + FBlendStackRequesterTrait as additive).
+ */	
+USTRUCT(meta = (DisplayName = "Blend Stack Requester", ShowTooltip=true))
 struct FAnimNextBlendStackRequesterTraitSharedData : public FAnimNextTraitSharedData
 {
 	GENERATED_BODY()
diff --git a/Engine/Plugins/Experimental/UAF/UAFAnimGraph/Source/UAFAnimGraph/Internal/Traits/CallFunction.h b/Engine/Plugins/Experimental/UAF/UAFAnimGraph/Source/UAFAnimGraph/Internal/Traits/CallFunction.h
--- a/Engine/Plugins/Experimental/UAF/UAFAnimGraph/Source/UAFAnimGraph/Internal/Traits/CallFunction.h
+++ b/Engine/Plugins/Experimental/UAF/UAFAnimGraph/Source/UAFAnimGraph/Internal/Traits/CallFunction.h
@@ -21,7 +21,8 @@ enum class EAnimNextCallFunctionCallSite : uint8
 	PostUpdate,
 };
 
-USTRUCT(meta = (DisplayName = "Call Function"))
+/** A trait that calls a function at the specified update point */
+USTRUCT(meta = (DisplayName = "Call Function", AllowDuplicates = "true"))
 struct FAnimNextCallFunctionSharedData : public FAnimNextTraitSharedData
 {
 	GENERATED_BODY()
@@ -40,7 +41,7 @@ struct FAnimNextCallFunctionSharedData : public FAnimNextTraitSharedData
 	FName FunctionEvent;
 
 	/** The call site to use */
-	UPROPERTY(EditAnywhere, Category = "Call Function")
+	UPROPERTY(EditAnywhere, Category = "Call Function", meta=(PinHiddenByDefault))
 	EAnimNextCallFunctionCallSite CallSite = EAnimNextCallFunctionCallSite::BecomeRelevant;
 
 	// Latent pin support boilerplate
diff --git a/Engine/Plugins/Experimental/UAF/UAFAnimGraph/Source/UAFAnimGraph/Private/PlayAnim/PlayAnimSlotTrait.h b/Engine/Plugins/Experimental/UAF/UAFAnimGraph/Source/UAFAnimGraph/Private/PlayAnim/PlayAnimSlotTrait.h
--- a/Engine/Plugins/Experimental/UAF/UAFAnimGraph/Source/UAFAnimGraph/Private/PlayAnim/PlayAnimSlotTrait.h
+++ b/Engine/Plugins/Experimental/UAF/UAFAnimGraph/Source/UAFAnimGraph/Private/PlayAnim/PlayAnimSlotTrait.h
@@ -17,7 +17,7 @@
 
 #include "PlayAnimSlotTrait.generated.h"
 
-USTRUCT(meta = (DisplayName = "PlayAnim Slot"))
+USTRUCT(meta = (DisplayName = "PlayAnim Slot", Hidden))
 struct FAnimNextPlayAnimSlotTraitSharedData : public FAnimNextTraitSharedData
 {
 	GENERATED_BODY()
diff --git a/Engine/Plugins/Experimental/UAF/UAFAnimGraph/Source/UAFAnimGraph/Private/Traits/ApplyAdditiveTrait.h b/Engine/Plugins/Experimental/UAF/UAFAnimGraph/Source/UAFAnimGraph/Private/Traits/ApplyAdditiveTrait.h
--- a/Engine/Plugins/Experimental/UAF/UAFAnimGraph/Source/UAFAnimGraph/Private/Traits/ApplyAdditiveTrait.h
+++ b/Engine/Plugins/Experimental/UAF/UAFAnimGraph/Source/UAFAnimGraph/Private/Traits/ApplyAdditiveTrait.h
@@ -13,7 +13,11 @@
 
 #include "ApplyAdditiveTrait.generated.h"
 
-USTRUCT(meta = (DisplayName = "Apply Additive"))
+/**
+ * A trait that can apply a mesh or local space additive to this trait stack.
+ * Ex: LookAt's, minor hit reactions, up-down floating, etc.
+ */
+USTRUCT(meta = (DisplayName = "Apply Additive", ShowTooltip=true))
 struct FAnimNextApplyAdditiveTraitSharedData : public FAnimNextTraitSharedData
 {
 	GENERATED_BODY()
diff --git a/Engine/Plugins/Experimental/UAF/UAFAnimGraph/Source/UAFAnimGraph/Private/Traits/BlendInertializer.h b/Engine/Plugins/Experimental/UAF/UAFAnimGraph/Source/UAFAnimGraph/Private/Traits/BlendInertializer.h
--- a/Engine/Plugins/Experimental/UAF/UAFAnimGraph/Source/UAFAnimGraph/Private/Traits/BlendInertializer.h
+++ b/Engine/Plugins/Experimental/UAF/UAFAnimGraph/Source/UAFAnimGraph/Private/Traits/BlendInertializer.h
@@ -11,7 +11,11 @@
 
 #include "BlendInertializer.generated.h"
 
-USTRUCT(meta = (DisplayName = "Blend Inertializer Core"))
+/**
+ * A trait that converts a normal smooth blend into an inertializing blend.
+ * It only implements the inertializing logic, it queries its required arguments using the IInertializerBlend interface.
+ */
+USTRUCT(meta = (DisplayName = "Blend Inertializer Core", ShowTooltip=true))
 struct FAnimNextBlendInertializerCoreTraitSharedData : public FAnimNextTraitSharedData
 {
 	GENERATED_BODY()
@@ -19,7 +23,11 @@ struct FAnimNextBlendInertializerCoreTraitSharedData : public FAnimNextTraitShar
 	// This struct is empty but required so that we can give a nice display name to the trait
 };
 
-USTRUCT(meta = (DisplayName = "Blend Inertializer", Category = "Default"))
+/**
+ * A trait that converts a normal smooth blend into an inertializing blend.
+ * This trait implements both the logic and contains the arguments necessary.
+ */
+USTRUCT(meta = (DisplayName = "Blend Inertializer", ShowTooltip=true))
 struct FAnimNextBlendInertializerTraitSharedData : public FAnimNextBlendInertializerCoreTraitSharedData
 {
 	GENERATED_BODY()
diff --git a/Engine/Plugins/Experimental/UAF/UAFAnimGraph/Source/UAFAnimGraph/Private/Traits/BlendLayer.h b/Engine/Plugins/Experimental/UAF/UAFAnimGraph/Source/UAFAnimGraph/Private/Traits/BlendLayer.h
--- a/Engine/Plugins/Experimental/UAF/UAFAnimGraph/Source/UAFAnimGraph/Private/Traits/BlendLayer.h
+++ b/Engine/Plugins/Experimental/UAF/UAFAnimGraph/Source/UAFAnimGraph/Private/Traits/BlendLayer.h
@@ -14,7 +14,8 @@
 
 #include "BlendLayer.generated.h"
 
-USTRUCT(meta = (DisplayName = "Blend Layer"))
+/** A trait that can blend a layer into an input. */
+USTRUCT(meta = (DisplayName = "Blend Layer", ShowTooltip=true))
 struct FAnimNextBlendLayerTraitSharedData : public FAnimNextTraitSharedData
 {
 	GENERATED_BODY()
diff --git a/Engine/Plugins/Experimental/UAF/UAFAnimGraph/Source/UAFAnimGraph/Private/Traits/BlendSmootherPerBone.h b/Engine/Plugins/Experimental/UAF/UAFAnimGraph/Source/UAFAnimGraph/Private/Traits/BlendSmootherPerBone.h
--- a/Engine/Plugins/Experimental/UAF/UAFAnimGraph/Source/UAFAnimGraph/Private/Traits/BlendSmootherPerBone.h
+++ b/Engine/Plugins/Experimental/UAF/UAFAnimGraph/Source/UAFAnimGraph/Private/Traits/BlendSmootherPerBone.h
@@ -17,7 +17,8 @@
 
 #include "BlendSmootherPerBone.generated.h"
 
-USTRUCT(meta = (DisplayName = "Blend Smoother Per Bone Core"))
+/** An additive trait that smoothly blends with per-bone weights. */
+USTRUCT(meta = (DisplayName = "Blend Smoother Per Bone Core", ShowTooltip=true))
 struct FAnimNextBlendSmootherPerBoneCoreTraitSharedData : public FAnimNextTraitSharedData
 {
 	GENERATED_BODY()
@@ -33,7 +34,11 @@ struct FAnimNextBlendProfilePerChildProviderTraitSharedData : public FAnimNextBl
 	TArray<TObjectPtr<UHierarchyTable>> BlendProfiles;
 };
 
-USTRUCT(meta = (DisplayName = "Blend Profile Provider"))
+/**
+ * Serves firstly as a passthrough node that propagates blend profiles of child traits.
+ * If no child trait provides a blend profile then uses the held blend profile.
+ */
+USTRUCT(meta = (DisplayName = "Blend Profile Provider", ShowTooltip=true))
 struct FAnimNextBlendProfileProviderTraitSharedData : public FAnimNextTraitSharedData
 {
 	GENERATED_BODY()
diff --git a/Engine/Plugins/Experimental/UAF/UAFAnimGraph/Source/UAFAnimGraph/Private/Traits/BlendTwoWay.h b/Engine/Plugins/Experimental/UAF/UAFAnimGraph/Source/UAFAnimGraph/Private/Traits/BlendTwoWay.h
--- a/Engine/Plugins/Experimental/UAF/UAFAnimGraph/Source/UAFAnimGraph/Private/Traits/BlendTwoWay.h
+++ b/Engine/Plugins/Experimental/UAF/UAFAnimGraph/Source/UAFAnimGraph/Private/Traits/BlendTwoWay.h
@@ -12,7 +12,8 @@
 
 #include "BlendTwoWay.generated.h"
 
-USTRUCT(meta = (DisplayName = "Blend Two Way"))
+/** A trait that can blend two inputs. */
+USTRUCT(meta = (DisplayName = "Blend Two Way", ShowTooltip=true))
 struct FAnimNextBlendTwoWayTraitSharedData : public FAnimNextTraitSharedData
 {
 	GENERATED_BODY()
diff --git a/Engine/Plugins/Experimental/UAF/UAFAnimGraph/Source/UAFAnimGraph/Private/Traits/InertializationDeadBlend.h b/Engine/Plugins/Experimental/UAF/UAFAnimGraph/Source/UAFAnimGraph/Private/Traits/InertializationDeadBlend.h
--- a/Engine/Plugins/Experimental/UAF/UAFAnimGraph/Source/UAFAnimGraph/Private/Traits/InertializationDeadBlend.h
+++ b/Engine/Plugins/Experimental/UAF/UAFAnimGraph/Source/UAFAnimGraph/Private/Traits/InertializationDeadBlend.h
@@ -16,7 +16,8 @@
 
 class UCurveFloat;
 
-USTRUCT(meta = (DisplayName = "Dead Blending"))
+/** A trait that inertializes animation by extrapolating from the point of transition. */
+USTRUCT(meta = (DisplayName = "Dead Blending", ShowTooltip=true))
 struct FAnimNextInertializationDeadBlendTraitSharedData : public FAnimNextTraitSharedData
 {
 	GENERATED_BODY()
diff --git a/Engine/Plugins/Experimental/UAF/UAFAnimGraph/Source/UAFAnimGraph/Private/Traits/InputPoseTrait.h b/Engine/Plugins/Experimental/UAF/UAFAnimGraph/Source/UAFAnimGraph/Private/Traits/InputPoseTrait.h
--- a/Engine/Plugins/Experimental/UAF/UAFAnimGraph/Source/UAFAnimGraph/Private/Traits/InputPoseTrait.h
+++ b/Engine/Plugins/Experimental/UAF/UAFAnimGraph/Source/UAFAnimGraph/Private/Traits/InputPoseTrait.h
@@ -10,7 +10,7 @@
 
 #include "InputPoseTrait.generated.h"
 
-USTRUCT(meta = (DisplayName = "Input Pose"))
+USTRUCT(meta = (DisplayName = "Input Pose", ShowTooltip=true))
 struct FAnimNextInputPoseTraitSharedData : public FAnimNextTraitSharedData
 {
 	GENERATED_BODY()
diff --git a/Engine/Plugins/Experimental/UAF/UAFAnimGraph/Source/UAFAnimGraph/Private/Traits/ModifyCurveTrait.h b/Engine/Plugins/Experimental/UAF/UAFAnimGraph/Source/UAFAnimGraph/Private/Traits/ModifyCurveTrait.h
--- a/Engine/Plugins/Experimental/UAF/UAFAnimGraph/Source/UAFAnimGraph/Private/Traits/ModifyCurveTrait.h
+++ b/Engine/Plugins/Experimental/UAF/UAFAnimGraph/Source/UAFAnimGraph/Private/Traits/ModifyCurveTrait.h
@@ -39,7 +39,8 @@ struct FModifyCurveParameters
 	
 };
 
-USTRUCT(meta = (DisplayName = "Modify Curve"))
+/** A trait that can write to a specific float curve value */
+USTRUCT(meta = (DisplayName = "Modify Curve", ShowTooltip=true, AllowDuplicates = "true"))
 struct FAnimNextModifyCurveTraitSharedData : public FAnimNextTraitSharedData
 {
 	GENERATED_BODY()
diff --git a/Engine/Plugins/Experimental/UAF/UAFAnimGraph/Source/UAFAnimGraph/Private/Traits/Passthrough.h b/Engine/Plugins/Experimental/UAF/UAFAnimGraph/Source/UAFAnimGraph/Private/Traits/Passthrough.h
--- a/Engine/Plugins/Experimental/UAF/UAFAnimGraph/Source/UAFAnimGraph/Private/Traits/Passthrough.h
+++ b/Engine/Plugins/Experimental/UAF/UAFAnimGraph/Source/UAFAnimGraph/Private/Traits/Passthrough.h
@@ -10,7 +10,8 @@
 
 #include "Passthrough.generated.h"
 
-USTRUCT(meta = (DisplayName = "Passthough"))
+/** A trait that passes through the input without modification. */
+USTRUCT(meta = (DisplayName = "Passthrough", ShowTooltip=true))
 struct FAnimNextPassthroughSharedData : public FAnimNextTraitSharedData
 {
 	GENERATED_BODY()
diff --git a/Engine/Plugins/Experimental/UAF/UAFAnimGraph/Source/UAFAnimGraph/Private/Traits/ReferencePoseTrait.h b/Engine/Plugins/Experimental/UAF/UAFAnimGraph/Source/UAFAnimGraph/Private/Traits/ReferencePoseTrait.h
--- a/Engine/Plugins/Experimental/UAF/UAFAnimGraph/Source/UAFAnimGraph/Private/Traits/ReferencePoseTrait.h
+++ b/Engine/Plugins/Experimental/UAF/UAFAnimGraph/Source/UAFAnimGraph/Private/Traits/ReferencePoseTrait.h
@@ -21,7 +21,8 @@ enum class EAnimNextReferencePoseType : int32
 	AdditiveIdentity,
 };
 
-USTRUCT(meta = (DisplayName = "Reference Pose"))
+/** A trait that outputs a reference pose. */
+USTRUCT(meta = (DisplayName = "Reference Pose", ShowTooltip=true))
 struct FAnimNextReferencePoseTraitSharedData : public FAnimNextTraitSharedData
 {
 	GENERATED_BODY()
diff --git a/Engine/Plugins/Experimental/UAF/UAFAnimGraph/Source/UAFAnimGraph/Private/Traits/SubGraphHost.h b/Engine/Plugins/Experimental/UAF/UAFAnimGraph/Source/UAFAnimGraph/Private/Traits/SubGraphHost.h
--- a/Engine/Plugins/Experimental/UAF/UAFAnimGraph/Source/UAFAnimGraph/Private/Traits/SubGraphHost.h
+++ b/Engine/Plugins/Experimental/UAF/UAFAnimGraph/Source/UAFAnimGraph/Private/Traits/SubGraphHost.h
@@ -14,7 +14,8 @@
 
 #include "SubGraphHost.generated.h"
 
-USTRUCT(meta = (DisplayName = "SubGraph"))
+/** A trait that hosts and manages a sub-graph instance. */
+USTRUCT(meta = (DisplayName = "SubGraph", ShowTooltip=true))
 struct FAnimNextSubGraphHostTraitSharedData : public FAnimNextTraitSharedData
 {
 	GENERATED_BODY()
diff --git a/Engine/Plugins/Experimental/UAF/UAFAnimGraph/Source/UAFAnimGraph/Public/Traits/BlendSmoother.h b/Engine/Plugins/Experimental/UAF/UAFAnimGraph/Source/UAFAnimGraph/Public/Traits/BlendSmoother.h
--- a/Engine/Plugins/Experimental/UAF/UAFAnimGraph/Source/UAFAnimGraph/Public/Traits/BlendSmoother.h
+++ b/Engine/Plugins/Experimental/UAF/UAFAnimGraph/Source/UAFAnimGraph/Public/Traits/BlendSmoother.h
@@ -15,7 +15,11 @@
 
 class UCurveFloat;
 
-USTRUCT(meta = (DisplayName = "Blend Smoother Core"))
+/**
+ * A trait that smoothly blends between discrete states over time.
+ * It only implements the smoothing logic, it queries its required arguments using the ISmoothBlend interface.
+ */
+USTRUCT(meta = (DisplayName = "Blend Smoother Core", ShowTooltip=true))
 struct FAnimNextBlendSmootherCoreTraitSharedData : public FAnimNextTraitSharedData
 {
 	GENERATED_BODY()
@@ -29,6 +33,10 @@ namespace UE::UAF
 	using FBlendSmootherCoreData = FAnimNextBlendSmootherCoreTraitSharedData; 
 }
 
+/**
+ * A trait that smoothly blends between discrete states over time.
+ * This trait implements both the logic and contains the arguments necessary.
+ */
 USTRUCT(meta = (DisplayName = "Blend Smoother"))
 struct FAnimNextBlendSmootherTraitSharedData : public FAnimNextBlendSmootherCoreTraitSharedData
 {
diff --git a/Engine/Plugins/Experimental/UAF/UAFAnimGraph/Source/UAFAnimGraph/Public/Traits/BlendSpacePlayerTraitData.h b/Engine/Plugins/Experimental/UAF/UAFAnimGraph/Source/UAFAnimGraph/Public/Traits/BlendSpacePlayerTraitData.h
--- a/Engine/Plugins/Experimental/UAF/UAFAnimGraph/Source/UAFAnimGraph/Public/Traits/BlendSpacePlayerTraitData.h
+++ b/Engine/Plugins/Experimental/UAF/UAFAnimGraph/Source/UAFAnimGraph/Public/Traits/BlendSpacePlayerTraitData.h
@@ -9,7 +9,8 @@
 
 #include "BlendSpacePlayerTraitData.generated.h"
 
-USTRUCT(meta = (DisplayName = "Blend Space Player"))
+/** A trait that can play a blend space */
+USTRUCT(meta = (DisplayName = "Blend Space Player", ShowTooltip=true))
 struct FAnimNextBlendSpacePlayerTraitSharedData : public FAnimNextTraitSharedData
 {
 	GENERATED_BODY()
diff --git a/Engine/Plugins/Experimental/UAF/UAFAnimGraph/Source/UAFAnimGraph/Public/Traits/NotifyDispatcherTraitData.h b/Engine/Plugins/Experimental/UAF/UAFAnimGraph/Source/UAFAnimGraph/Public/Traits/NotifyDispatcherTraitData.h
--- a/Engine/Plugins/Experimental/UAF/UAFAnimGraph/Source/UAFAnimGraph/Public/Traits/NotifyDispatcherTraitData.h
+++ b/Engine/Plugins/Experimental/UAF/UAFAnimGraph/Source/UAFAnimGraph/Public/Traits/NotifyDispatcherTraitData.h
@@ -7,7 +7,8 @@
 #include "TraitCore/TraitBinding.h"
 #include "NotifyDispatcherTraitData.generated.h"
 
-USTRUCT(meta = (DisplayName = "Notify Dispatcher"))
+/** A trait that dispatches notifies according to a timeline advancing */
+USTRUCT(meta = (DisplayName = "Notify Dispatcher", ShowTooltip=true))
 struct FAnimNextNotifyDispatcherTraitSharedData : public FAnimNextTraitSharedData
 {
 	GENERATED_BODY()
diff --git a/Engine/Plugins/Experimental/UAF/UAFAnimGraph/Source/UAFAnimGraph/Public/Traits/PassthroughBlendTrait.h b/Engine/Plugins/Experimental/UAF/UAFAnimGraph/Source/UAFAnimGraph/Public/Traits/PassthroughBlendTrait.h
--- a/Engine/Plugins/Experimental/UAF/UAFAnimGraph/Source/UAFAnimGraph/Public/Traits/PassthroughBlendTrait.h
+++ b/Engine/Plugins/Experimental/UAF/UAFAnimGraph/Source/UAFAnimGraph/Public/Traits/PassthroughBlendTrait.h
@@ -11,28 +11,26 @@
 // --- ---
 #include "PassthroughBlendTrait.generated.h"
 
-/*
-* Passthrough Blend :
-* It will generate an output using the base trait input keyframe and the selected alpha blend type (float, bool or curve)
-* Requires a base trait with a single input and PostEvaluate task creation (currently used to blend Control Rig Trait)
-* 
-* If alpha type is Float or Bool
-* Depending on the weight value
-* - 0 (not relevant)          : The input "as is" will be the output (passthrough). No call to PostEvaluate on the base will be performed.
-* - 1 (full weight)           : The output of the base will be used as the output without blending by calling base PostEvaluate
-* - Any other relevant weight : Input will be duplicated (with a task) in Post Evaluate
-*								The base Post Evaluate will be called to execute on the copy (normally will create a task)
-*								Both keyframes (original and copy with base modifications) will be blended using a two keyframes blend task, passing the calculated alpha blend weight
-*
-* If alpha type is Curve
-* The curve value is not known until we are in the task (because we need the keyframe to extract the curve value)
-* Currently, the process will generate a full blend :	A task will added to duplicate input keyframe at PostEvaluate
-*														Base PostEvaluate will be called to allow task creation
-*														Resulting keyframes will be blended with a two keyframes blend task, passing the curve data so the alpha can be obtained
-* If we try to get current alpha value of a Curve blend type, the returned value will be 1
-*/
-
-USTRUCT(meta = (DisplayName = "Passthrough Blend"))
+/**
+ * It will generate an output using the base trait input keyframe and the selected alpha blend type (float, bool or curve)
+ * Requires a base trait with a single input and PostEvaluate task creation (currently used to blend Control Rig Trait)
+ * 
+ * If alpha type is Float or Bool
+ * Depending on the weight value
+ * - 0 (not relevant)          : The input "as is" will be the output (passthrough). No call to PostEvaluate on the base will be performed.
+ * - 1 (full weight)           : The output of the base will be used as the output without blending by calling base PostEvaluate
+ * - Any other relevant weight : Input will be duplicated (with a task) in Post Evaluate
+ *								The base Post Evaluate will be called to execute on the copy (normally will create a task)
+ *								Both keyframes (original and copy with base modifications) will be blended using a two keyframes blend task, passing the calculated alpha blend weight
+ *
+ * If alpha type is Curve
+ * The curve value is not known until we are in the task (because we need the keyframe to extract the curve value)
+ * Currently, the process will generate a full blend :	A task will added to duplicate input keyframe at PostEvaluate
+ *														Base PostEvaluate will be called to allow task creation
+ *														Resulting keyframes will be blended with a two keyframes blend task, passing the curve data so the alpha can be obtained
+ * If we try to get current alpha value of a Curve blend type, the returned value will be 1
+ */
+USTRUCT(meta = (DisplayName = "Passthrough Blend", ShowTooltip=true))
 struct FPassthroughBlendTraitSharedData : public FAnimNextTraitSharedData
 {
 	GENERATED_BODY()
diff --git a/Engine/Plugins/Experimental/UAF/UAFAnimGraph/Source/UAFAnimGraph/Public/Traits/SequencePlayerTraitData.h b/Engine/Plugins/Experimental/UAF/UAFAnimGraph/Source/UAFAnimGraph/Public/Traits/SequencePlayerTraitData.h
--- a/Engine/Plugins/Experimental/UAF/UAFAnimGraph/Source/UAFAnimGraph/Public/Traits/SequencePlayerTraitData.h
+++ b/Engine/Plugins/Experimental/UAF/UAFAnimGraph/Source/UAFAnimGraph/Public/Traits/SequencePlayerTraitData.h
@@ -7,7 +7,8 @@
 #include "Animation/AnimSequence.h"
 #include "SequencePlayerTraitData.generated.h"
 
-USTRUCT(meta = (DisplayName = "Sequence Player"))
+/** A trait that can play an animation sequence. */
+USTRUCT(meta = (DisplayName = "Sequence Player", ShowTooltip=true))
 struct FAnimNextSequencePlayerTraitSharedData : public FAnimNextTraitSharedData
 {
 	GENERATED_BODY()
diff --git a/Engine/Plugins/Experimental/UAF/UAFAnimGraph/Source/UAFAnimGraph/Public/Traits/SynchronizeUsingGroupsTraitData.h b/Engine/Plugins/Experimental/UAF/UAFAnimGraph/Source/UAFAnimGraph/Public/Traits/SynchronizeUsingGroupsTraitData.h
--- a/Engine/Plugins/Experimental/UAF/UAFAnimGraph/Source/UAFAnimGraph/Public/Traits/SynchronizeUsingGroupsTraitData.h
+++ b/Engine/Plugins/Experimental/UAF/UAFAnimGraph/Source/UAFAnimGraph/Public/Traits/SynchronizeUsingGroupsTraitData.h
@@ -8,6 +8,7 @@
 #include "TraitInterfaces/IGroupSynchronization.h"
 #include "SynchronizeUsingGroupsTraitData.generated.h"
 
+/** A trait that synchronizes animation sequence playback using named groups. */
 USTRUCT(meta = (DisplayName = "Synchronize Using Groups"))
 struct FAnimNextSynchronizeUsingGroupsTraitSharedData : public FAnimNextTraitSharedData
 {
diff --git a/Engine/Plugins/Experimental/UAF/UAFAnimGraph/Source/UAFAnimGraphEditor/Private/AnimNextAnimGraphEditorModule.cpp b/Engine/Plugins/Experimental/UAF/UAFAnimGraph/Source/UAFAnimGraphEditor/Private/AnimNextAnimGraphEditorModule.cpp
--- a/Engine/Plugins/Experimental/UAF/UAFAnimGraph/Source/UAFAnimGraphEditor/Private/AnimNextAnimGraphEditorModule.cpp
+++ b/Engine/Plugins/Experimental/UAF/UAFAnimGraph/Source/UAFAnimGraphEditor/Private/AnimNextAnimGraphEditorModule.cpp
@@ -118,7 +118,7 @@ class FAnimNextAnimGraphEditorModule : public IModuleInterface
 
 		WorkspaceEditorModule.RegisterObjectDocumentType(FTopLevelAssetPath(TEXT("/Script/UAFAnimGraph.AnimNextAnimationGraph")), AnimNextAnimationGraphDocumentArgs);
 		
-		WorkspaceEditorModule.OnRegisterTabsForEditor().AddLambda([](FWorkflowAllowedTabSet& TabFactories, const TSharedRef<FTabManager>& InTabManager, TSharedPtr<UE::Workspace::IWorkspaceEditor> InEditorPtr)
+/*		WorkspaceEditorModule.OnRegisterTabsForEditor().AddLambda([](FWorkflowAllowedTabSet& TabFactories, const TSharedRef<FTabManager>& InTabManager, TSharedPtr<UE::Workspace::IWorkspaceEditor> InEditorPtr)
 		{
 			TSharedRef<FTraitEditorTabSummoner> TraitEditorTabSummoner = MakeShared<FTraitEditorTabSummoner>(InEditorPtr);
 			TabFactories.RegisterFactory(TraitEditorTabSummoner);
@@ -129,7 +129,7 @@ class FAnimNextAnimGraphEditorModule : public IModuleInterface
 		{
 			FTabManager::FTab TraitEditorTab(FTabId(TraitEditorTabName), ETabState::ClosedTab);
 			InLayoutExtender.ExtendLayout(FTabId(Workspace::WorkspaceTabs::TopRightDocumentArea), ELayoutExtensionPosition::After, TraitEditorTab);
-		});
+		});*/
 		
 		const TSharedPtr<FAnimNextAnimationGraphItemDetails> AssetItemDetails = MakeShared<FAnimNextAnimationGraphItemDetails>();
 		WorkspaceEditorModule.RegisterWorkspaceItemDetails(Workspace::FOutlinerItemDetailsId(FAnimNextAnimationGraphOutlinerData::StaticStruct()->GetFName()), AssetItemDetails);
diff --git a/Engine/Plugins/Experimental/UAF/UAFAnimGraph/Source/UAFAnimGraphEditor/Private/Graph/AnimNextAnimationGraphMenuExtensions.cpp b/Engine/Plugins/Experimental/UAF/UAFAnimGraph/Source/UAFAnimGraphEditor/Private/Graph/AnimNextAnimationGraphMenuExtensions.cpp
--- a/Engine/Plugins/Experimental/UAF/UAFAnimGraph/Source/UAFAnimGraphEditor/Private/Graph/AnimNextAnimationGraphMenuExtensions.cpp
+++ b/Engine/Plugins/Experimental/UAF/UAFAnimGraph/Source/UAFAnimGraphEditor/Private/Graph/AnimNextAnimationGraphMenuExtensions.cpp
@@ -98,69 +98,6 @@ void FAnimationGraphMenuExtensions::RegisterMenus()
 		if (UncookedOnly::FAnimGraphUtils::IsTraitStackNode(ModelNode))
 		{
 			UAnimNextTraitStackUnitNode* TraitStackNode = CastChecked<UAnimNextTraitStackUnitNode>(ModelNode);
-			{
-				FToolMenuSection& Section = InMenu->AddSection("AnimNextTraitNodeActions", LOCTEXT("AnimNextTraitNodeActionsMenuHeader", "Traits"));
-
-				auto BuildAddTraitContextMenu = [AnimNextEdGraphNode, ModelNode](UToolMenu* InSubMenu)
-				{
-					const FTraitRegistry& TraitRegistry = FTraitRegistry::Get();
-					TArray<const FTrait*> Traits = TraitRegistry.GetTraits();
-
-					URigVMController* VMController = AnimNextEdGraphNode->GetController();
-
-					const UScriptStruct* CppDecoratorStruct = FRigDecorator_AnimNextCppDecorator::StaticStruct();
-
-					for (const FTrait* Trait : Traits)
-					{
-						UScriptStruct* ScriptStruct = Trait->GetTraitSharedDataStruct();
-				
-						FString DefaultValue;
-						{
-							const FRigDecorator_AnimNextCppDecorator DefaultCppDecoratorStructInstance;
-							FRigDecorator_AnimNextCppDecorator CppDecoratorStructInstance;
-							CppDecoratorStructInstance.DecoratorSharedDataStruct = ScriptStruct;
-				
-							if (!CppDecoratorStructInstance.CanBeAddedToNode(ModelNode, nullptr))
-							{
-								continue;	// This trait isn't supported on this node
-							}
-
-							FRigDecorator_AnimNextCppDecorator::StaticStruct()->ExportText(DefaultValue, &CppDecoratorStructInstance, &DefaultCppDecoratorStructInstance, nullptr, PPF_SerializedAsImportText, nullptr);
-						}
-				
-						FString DisplayNameMetadata;
-						ScriptStruct->GetStringMetaDataHierarchical(FRigVMStruct::DisplayNameMetaName, &DisplayNameMetadata);
-						const FString DisplayName = DisplayNameMetadata.IsEmpty() ? Trait->GetTraitName() : DisplayNameMetadata;
-						const FText ToolTip = ScriptStruct->GetToolTipText();
-				
-						FToolMenuEntry TraitEntry = FToolMenuEntry::InitMenuEntry(
-							*Trait->GetTraitName(),
-							FText::FromString(DisplayName),
-							ToolTip,
-							FSlateIcon(),
-							FUIAction(FExecuteAction::CreateLambda(
-								[Trait, VMController, ModelNode, CppDecoratorStruct, DefaultValue, DisplayName]()
-								{
-									VMController->AddTrait(
-										ModelNode->GetFName(),
-										*CppDecoratorStruct->GetPathName(),
-										*DisplayName,
-										DefaultValue, INDEX_NONE, true, true);
-								})
-							)
-						);
-
-						InSubMenu->AddMenuEntry(NAME_None, TraitEntry);
-					}
-				};
-				
-				Section.AddSubMenu(
-					TEXT("AddTraitMenu"),
-					LOCTEXT("AddTraitMenu", "Add Trait"),
-					LOCTEXT("AddTraitMenuTooltip", "Add the chosen trait to currently selected node"),
-					FNewToolMenuDelegate::CreateLambda(BuildAddTraitContextMenu));
-			}
-
 			{
 				FToolMenuSection& Section = InMenu->AddSection("AnimNextTemplateNodeActions", LOCTEXT("AnimNextTemplateNodeActionsMenuHeader", "Template"));
 
diff --git a/Engine/Plugins/Experimental/UAF/UAFAnimGraph/Source/UAFAnimGraphEditor/Private/TraitStackEditor.cpp b/Engine/Plugins/Experimental/UAF/UAFAnimGraph/Source/UAFAnimGraphEditor/Private/TraitStackEditor.cpp
--- a/Engine/Plugins/Experimental/UAF/UAFAnimGraph/Source/UAFAnimGraphEditor/Private/TraitStackEditor.cpp
+++ b/Engine/Plugins/Experimental/UAF/UAFAnimGraph/Source/UAFAnimGraphEditor/Private/TraitStackEditor.cpp
@@ -1,10 +1,27 @@
 ﻿// Copyright Epic Games, Inc. All Rights Reserved.
 
 #include "TraitStackEditor.h"
+
+#include "AnimNextController.h"
+#include "AnimNextEdGraphNode.h"
+#include "DetailCategoryBuilder.h"
+#include "DetailLayoutBuilder.h"
+#include "IDetailPropertyRow.h"
 #include "IWorkspaceEditor.h"
+#include "PropertyHandle.h"
+#include "SPositiveActionButton.h"
+#include "SSimpleComboButton.h"
+#include "StructViewerFilter.h"
+#include "StructViewerModule.h"
+#include "Framework/Application/SlateApplication.h"
 #include "Graph/TraitEditorTabSummoner.h"
+#include "Templates/UAFGraphNodeTemplate.h"
+#include "TraitCore/TraitInterfaceRegistry.h"
+#include "TraitCore/TraitRegistry.h"
 #include "Widgets/Docking/SDockTab.h"
 
+#define LOCTEXT_NAMESPACE "TraitStackEditor"
+
 namespace UE::UAF::Editor
 {
 
@@ -19,4 +36,810 @@ void FTraitStackEditor::SetTraitData(const TSharedRef<UE::Workspace::IWorkspaceE
 	}
 }
 
+namespace Private
+{
+	class FTraitStackDragDropOp : public FDecoratedDragDropOp
+	{
+	public:
+		DRAG_DROP_OPERATOR_TYPE(FTraitStackDragDropOp, FDecoratedDragDropOp)
+
+		static TSharedRef<FTraitStackDragDropOp> New(const TSharedRef<IPropertyHandle>& InPropertyHandle, FText InDisplayName, FVector2f InSize)
+		{
+			TSharedRef<FTraitStackDragDropOp> Operation = MakeShared<FTraitStackDragDropOp>();
+			Operation->WeakPropertyHandle = InPropertyHandle;
+			Operation->DisplayName = InDisplayName;
+			Operation->Size = FVector2f(FMath::Min(InSize.X, 300.0f), InSize.Y);
+			Operation->Construct();
+			return Operation;
+		}
+
+		virtual TSharedPtr<SWidget> GetDefaultDecorator() const override
+		{
+			return
+				SNew(SBox)
+				.WidthOverride(Size.X)
+				.HeightOverride(Size.Y)
+				[
+					SNew(SOverlay)
+					+SOverlay::Slot()
+					[
+						SNew(SBorder)
+						.BorderImage(FAppStyle::Get().GetBrush("Brushes.White"))
+						.BorderBackgroundColor(FStyleColors::Select.GetSpecifiedColor())
+					]
+					+SOverlay::Slot()
+					.Padding(1.0f)
+					[
+						SNew(SBorder)
+						.BorderImage(FAppStyle::GetBrush("ToolPanel.DarkGroupBorder"))
+						.Content()
+						[
+							SNew(SBox)
+							.HAlign(HAlign_Left)
+							.VAlign(VAlign_Center)
+							.Padding(3.0f)
+							[
+								SNew(STextBlock)
+								.Text(DisplayName)
+								.Font(IDetailLayoutBuilder::GetDetailFontBold())
+							]
+						]
+					]
+				];
+		}
+
+		TWeakPtr<IPropertyHandle> WeakPropertyHandle;
+		FText DisplayName;
+		FVector2f Size;
+	};
+
+	class STraitHeader : public SCompoundWidget
+	{
+		SLATE_BEGIN_ARGS(STraitHeader) {}
+
+		SLATE_END_ARGS()
+
+		void Construct(const FArguments& InArgs, IDetailCategoryBuilder& InCategory, IDetailPropertyRow& InPropertyRow, URigVMController* InController);
+
+		/** Retrieves a brush for rendering a drop indicator for the specified drop zone */
+		static const FSlateBrush* GetDropIndicatorBrush(EItemDropZone InItemDropZone)
+		{
+			const FTableRowStyle& TreeViewStyle = FAppStyle::GetWidgetStyle<FTableRowStyle>("DetailsView.TreeView.TableRow");
+			switch (InItemDropZone)
+			{
+			case EItemDropZone::AboveItem: return &TreeViewStyle.DropIndicator_Above;
+			default:
+			case EItemDropZone::OntoItem: return &TreeViewStyle.DropIndicator_Onto;
+			case EItemDropZone::BelowItem: return &TreeViewStyle.DropIndicator_Below;
+			};
+		}
+		
+		virtual int32 OnPaint( const FPaintArgs& Args, const FGeometry& AllottedGeometry, const FSlateRect& MyCullingRect, FSlateWindowElementList& OutDrawElements, int32 LayerId, const FWidgetStyle& InWidgetStyle, bool bParentEnabled ) const override
+		{
+			LayerId = SCompoundWidget::OnPaint(Args, AllottedGeometry, MyCullingRect, OutDrawElements, LayerId, InWidgetStyle, bParentEnabled);
+
+			if (ItemDropZone.IsSet())
+			{
+				// Draw feedback for user dropping an item above, below, or onto a row.
+				const FSlateBrush* DropIndicatorBrush = GetDropIndicatorBrush(ItemDropZone.GetValue());
+
+				FSlateDrawElement::MakeBox
+				(
+					OutDrawElements,
+					LayerId++,
+					AllottedGeometry.ToPaintGeometry(),
+					DropIndicatorBrush,
+					ESlateDrawEffect::None,
+					DropIndicatorBrush->GetTint(InWidgetStyle) * InWidgetStyle.GetColorAndOpacityTint()
+				);
+			}
+
+			Geometry = AllottedGeometry;
+			return LayerId;
+		}
+
+		/** @return the zone (above, onto, below) based on where the user is hovering over within the row */
+		static EItemDropZone ZoneFromPointerPosition(UE::Slate::FDeprecateVector2DParameter LocalPointerPos, UE::Slate::FDeprecateVector2DParameter LocalSize)
+		{
+			const float ZoneBoundarySu = FMath::Clamp(LocalSize.Y * 0.25f, 3.0f, 10.0f);
+			if (LocalPointerPos.Y < ZoneBoundarySu)
+			{
+				return EItemDropZone::AboveItem;
+			}
+			else if (LocalPointerPos.Y > LocalSize.Y - ZoneBoundarySu)
+			{
+				return EItemDropZone::BelowItem;
+			}
+			else
+			{
+				return EItemDropZone::OntoItem;
+			}
+		}
+
+		virtual void OnDragEnter(const FGeometry& MyGeometry, const FDragDropEvent& DragDropEvent) override
+		{
+			TSharedPtr<FTraitStackDragDropOp> DragOp = DragDropEvent.GetOperationAs<FTraitStackDragDropOp>();
+			if (!DragOp.IsValid())
+			{
+				return;
+			}
+
+			const FVector2f LocalPointerPos = MyGeometry.AbsoluteToLocal(DragDropEvent.GetScreenSpacePosition());
+			ItemDropZone = ZoneFromPointerPosition(LocalPointerPos, MyGeometry.GetLocalSize());
+		}
+
+		virtual void OnDragLeave(FDragDropEvent const& DragDropEvent) override
+		{
+			ItemDropZone = TOptional<EItemDropZone>();
+		}
+
+		virtual FReply OnDragOver(const FGeometry& MyGeometry, const FDragDropEvent& DragDropEvent) override
+		{
+			TSharedPtr<FTraitStackDragDropOp> DragOp = DragDropEvent.GetOperationAs<FTraitStackDragDropOp>();
+			if (!DragOp.IsValid())
+			{
+				return FReply::Unhandled();
+			}
+			
+			const FVector2f LocalPointerPos = MyGeometry.AbsoluteToLocal(DragDropEvent.GetScreenSpacePosition());
+			ItemDropZone = ZoneFromPointerPosition(LocalPointerPos, MyGeometry.GetLocalSize());
+
+			return FReply::Handled();
+		}
+
+		virtual FReply OnDrop(const FGeometry& MyGeometry, const FDragDropEvent& DragDropEvent) override
+		{
+			ON_SCOPE_EXIT { ItemDropZone = TOptional<EItemDropZone>(); };
+		
+			TSharedPtr<FTraitStackDragDropOp> DragOp = DragDropEvent.GetOperationAs<FTraitStackDragDropOp>();
+			if (!DragOp.IsValid())
+			{
+				return FReply::Unhandled();
+			}
+			
+			UAnimNextController* Controller = Cast<UAnimNextController>(WeakController.Get());
+			if (Controller == nullptr)
+			{
+				return FReply::Unhandled();
+			}
+			
+			TSharedPtr<IPropertyHandle> SourceHandle = DragOp->WeakPropertyHandle.Pin();
+			TSharedPtr<IPropertyHandle> TargetHandle = WeakPropertyHandle.Pin();
+
+			if (!SourceHandle.IsValid() || !TargetHandle.IsValid())
+			{
+				return FReply::Unhandled();
+			}
+
+			TSharedPtr<IPropertyHandleStruct> SourceStructHandle = SourceHandle->AsStruct();
+			TSharedPtr<IPropertyHandleStruct> TargetStructHandle = TargetHandle->AsStruct();
+			check(SourceStructHandle.IsValid() && TargetStructHandle.IsValid());
+
+			TSharedPtr<FStructOnScope> SourceStructData = SourceStructHandle->GetStructData();
+			TSharedPtr<FStructOnScope> TargetStructData = TargetStructHandle->GetStructData();
+
+			FRigVMControllerCompileBracketScope CompileBracketScope(Controller);
+			Controller->OpenUndoBracket(LOCTEXT("ReorderTraitsTransaction", "Reorder Traits").ToString());
+
+			for (TWeakObjectPtr<UAnimNextEdGraphNode> WeakNode : Nodes)
+			{
+				if (UAnimNextEdGraphNode* Node = WeakNode.Get())
+				{
+					TArray<URigVMPin*> TraitPins = Node->GetModelNode()->GetTraitPins();
+					URigVMPin* SourcePin = nullptr;
+					int32 SourcePinIndex = INDEX_NONE;
+					for (URigVMPin* Pin : TraitPins)
+					{
+						TSharedPtr<FStructOnScope> TraitScope = Pin->GetTraitInstance();
+						if (TraitScope.IsValid())
+						{
+							const FRigVMTrait* VMTrait = (FRigVMTrait*)TraitScope->GetStructMemory();
+							if (VMTrait->GetTraitSharedDataStruct() == SourceStructData->GetStruct())
+							{
+								SourcePin = Pin;
+								SourcePinIndex = Pin->GetPinIndex();
+								break;
+							}
+						}
+					}
+
+					if (SourcePin == nullptr)
+					{
+						continue;
+					}
+
+					int32 TargetPinIndex = INDEX_NONE;
+					for (URigVMPin* Pin : TraitPins)
+					{
+						TSharedPtr<FStructOnScope> TraitScope = Pin->GetTraitInstance();
+						if (TraitScope.IsValid())
+						{
+							const FRigVMTrait* VMTrait = (FRigVMTrait*)TraitScope->GetStructMemory();
+							if (VMTrait->GetTraitSharedDataStruct() == TargetStructData->GetStruct())
+							{
+								TargetPinIndex = Pin->GetPinIndex();
+								switch (ItemDropZone.GetValue())
+								{
+								case EItemDropZone::AboveItem:
+									TargetPinIndex = FMath::Max(0, TargetPinIndex - 1);
+									break;
+								case EItemDropZone::OntoItem:
+									break;
+								case EItemDropZone::BelowItem:
+									TargetPinIndex = FMath::Min(TraitPins.Num(), TargetPinIndex + 1);
+									break;
+								}
+
+								if (SourcePinIndex != TargetPinIndex)
+								{
+									FRigVMNodeLayout NodeLayout = Node->GetModelNode()->GetNodeLayout(true);
+									Controller->SetTraitPinIndex(Node->GetModelNodeName(), SourcePin->GetFName(), TargetPinIndex, true, true);
+									if(NodeLayout.IsValid())
+									{
+										Controller->SetNodeLayout(Node->GetModelNodeName(), NodeLayout, true, true);
+									}
+								}
+								break;
+							}
+						}
+					}
+				}
+			}
+
+			Controller->CloseUndoBracket();
+
+			return FReply::Handled();
+		}
+
+		mutable FGeometry Geometry;
+		TOptional<EItemDropZone> ItemDropZone;
+		TWeakPtr<IPropertyHandle> WeakPropertyHandle;
+		TWeakObjectPtr<URigVMController> WeakController;
+		TArray<TWeakObjectPtr<UAnimNextEdGraphNode>> Nodes;
+		const UScriptStruct* SharedDataStruct; 
+	};
+
+	class FStructFilter : public IStructViewerFilter
+	{
+	public:
+		FStructFilter(TConstArrayView<TWeakObjectPtr<UAnimNextEdGraphNode>> InNodes, ETraitMode InMode)
+			: Nodes(InNodes)
+			, Mode(InMode)
+		{
+		}
+		
+		virtual bool IsStructAllowed(const FStructViewerInitializationOptions& InInitOptions, const UScriptStruct* InStruct, TSharedRef<FStructViewerFilterFuncs> InFilterFuncs) override
+		{
+			if (InStruct->HasMetaData(TEXT("Hidden")))
+			{
+				return false;
+			}
+
+			if (!InStruct->IsChildOf(FAnimNextTraitSharedData::StaticStruct()))
+			{
+				return false;
+			}
+	
+			const FTraitRegistry& TraitRegistry = FTraitRegistry::Get();
+			const FTrait* Trait = TraitRegistry.Find(InStruct);
+			if (Trait == nullptr)
+			{
+				return false;
+			}
+
+			if (Trait->GetTraitMode() != Mode)
+			{
+				return false;
+			}
+
+			// Check for duplicates unless they are allowed
+			if (InStruct->GetMetaData(TEXT("AllowDuplicates")) != TEXT("true"))
+			{
+				for (TWeakObjectPtr<UAnimNextEdGraphNode> WeakNode : Nodes)
+				{
+					if (UAnimNextEdGraphNode* Node = WeakNode.Get())
+					{
+						TArray<URigVMPin*> TraitPins = Node->GetModelNode()->GetTraitPins();
+						for (URigVMPin* Pin : TraitPins)
+						{
+							TSharedPtr<FStructOnScope> TraitScope = Pin->GetTraitInstance();
+							if (TraitScope.IsValid())
+							{
+								const FRigVMTrait* VMTrait = (FRigVMTrait*)TraitScope->GetStructMemory();
+								if (VMTrait->GetTraitSharedDataStruct() == InStruct)
+								{
+									return false;
+								}
+							}
+						}
+					}
+				}
+			}
+
+			return true;
+		}
+
+		virtual bool IsUnloadedStructAllowed(const FStructViewerInitializationOptions& InInitOptions, const FSoftObjectPath& InStructPath, TSharedRef<class FStructViewerFilterFuncs> InFilterFuncs)
+		{
+			return false;
+		};
+
+		TConstArrayView<TWeakObjectPtr<UAnimNextEdGraphNode>> Nodes;
+		ETraitMode Mode;
+	};
+
+	class SDragHandle : public SCompoundWidget
+	{
+		SLATE_BEGIN_ARGS(SDragHandle) {}
+
+		SLATE_ARGUMENT(FText, DisplayName)
+
+		SLATE_END_ARGS()
+
+		void Construct(const FArguments& InArgs, const TSharedRef<IPropertyHandle>& InPropertyHandle, const TSharedRef<STraitHeader>& InContainerWidget)
+		{
+			DisplayName = InArgs._DisplayName;
+			WeakPropertyHandle = InPropertyHandle;
+			WeakContainerWidget = InContainerWidget;
+
+			ChildSlot
+			[
+				SNew(SHorizontalBox)
+				.Cursor(EMouseCursor::GrabHand)
+				+SHorizontalBox::Slot()
+				.AutoWidth()
+				.VAlign(VAlign_Center)
+				[
+					SNew(SImage)
+					.Visibility_Lambda([this]()
+					{
+						if( const TSharedPtr<SWidget> Widget = WeakContainerWidget.Pin())
+						{
+							return Widget->IsHovered() ? EVisibility::Visible : EVisibility::Hidden;
+						}
+						return EVisibility::Hidden;
+					})
+					.Image(FCoreStyle::Get().GetBrush("VerticalBoxDragIndicatorShort"))
+				]
+				+SHorizontalBox::Slot()
+				.AutoWidth()
+				.VAlign(VAlign_Center)
+				.Padding(3.0f)
+				[
+					SNew(STextBlock)
+					.Text(DisplayName)
+					.Font(IDetailLayoutBuilder::GetDetailFontBold())
+				]
+			];
+		}
+
+		virtual FReply OnMouseButtonDown(const FGeometry& MyGeometry, const FPointerEvent& MouseEvent) override
+		{
+			if (!MouseEvent.IsMouseButtonDown(EKeys::LeftMouseButton))
+			{
+				return FReply::Unhandled();
+			}
+
+			return FReply::Handled().DetectDrag(SharedThis(this), EKeys::LeftMouseButton);
+		};
+		
+		virtual FReply OnDragDetected(const FGeometry& MyGeometry, const FPointerEvent& MouseEvent) override
+		{
+			if (!MouseEvent.IsMouseButtonDown(EKeys::LeftMouseButton))
+			{
+				return FReply::Unhandled();
+			}
+
+			TSharedPtr<IPropertyHandle> PropertyHandle = WeakPropertyHandle.Pin();
+			if (!PropertyHandle.IsValid())
+			{
+				return FReply::Unhandled();
+			}
+
+			TSharedPtr<STraitHeader> ContainerWidget = WeakContainerWidget.Pin();
+			if (!ContainerWidget.IsValid())
+			{
+				return FReply::Unhandled();
+			}
+
+			TSharedPtr<FDragDropOperation> DragDropOp = FTraitStackDragDropOp::New(PropertyHandle.ToSharedRef(), DisplayName, ContainerWidget->Geometry.GetLocalSize());
+			return FReply::Handled().BeginDragDrop(DragDropOp.ToSharedRef());
+		}
+
+		TWeakPtr<IPropertyHandle> WeakPropertyHandle;
+		TWeakPtr<STraitHeader> WeakContainerWidget;
+		FText DisplayName;
+	};
+	
+	void STraitHeader::Construct(const FArguments& InArgs, IDetailCategoryBuilder& InCategory, IDetailPropertyRow& InPropertyRow, URigVMController* InController)
+	{
+		WeakPropertyHandle = InPropertyRow.GetPropertyHandle();
+		WeakController = InController;
+		Nodes = InCategory.GetParentLayout().GetObjectsOfTypeBeingCustomized<UAnimNextEdGraphNode>();
+		
+		TSharedPtr<IPropertyHandleStruct> StructHandle = InPropertyRow.GetPropertyHandle()->AsStruct();
+		check(StructHandle.IsValid());
+
+		SharedDataStruct = Cast<UScriptStruct>(StructHandle->GetStructData()->GetStruct());
+		check(SharedDataStruct);
+
+		TSharedRef<Private::SDragHandle> HandleWidget =
+			SNew(Private::SDragHandle, InPropertyRow.GetPropertyHandle().ToSharedRef(), SharedThis(this))
+			.DisplayName(InCategory.GetDisplayName());
+
+		TSharedRef<SWidget> Widget =
+			SNew(SOverlay)
+			.ToolTipText_Lambda([this]()->FText
+			{
+				if (const FTrait* Trait = FTraitRegistry::Get().Find(SharedDataStruct))
+				{
+					static const FTextFormat EntryFormat(LOCTEXT("InterfaceListEntryFormat", "- {0}"));
+
+					FTextBuilder TraitInfoText;
+					TraitInfoText.AppendLine(SharedDataStruct->GetDisplayNameText());
+					TraitInfoText.AppendLine(FText::GetEmpty());
+					TraitInfoText.AppendLine(SharedDataStruct->GetToolTipText());
+
+					if (Trait->GetTraitInterfaces().Num() > 0)
+					{
+						TraitInfoText.AppendLine(FText::GetEmpty());
+						TraitInfoText.AppendLine(LOCTEXT("TraitInfoImplementedInterfaces", "Implements:"));
+						for (const FTraitInterfaceUID& ImplementedInterfaceUID : Trait->GetTraitInterfaces())
+						{
+							if (const ITraitInterface* ImplementedInterface = FTraitInterfaceRegistry::Get().Find(ImplementedInterfaceUID))
+							{
+								TraitInfoText.AppendLineFormat(EntryFormat, ImplementedInterface->GetDisplayName());
+							}
+						}
+					}
+
+					if (Trait->GetTraitRequiredInterfaces().Num() > 0)
+					{
+						TraitInfoText.AppendLine(FText::GetEmpty());
+						TraitInfoText.AppendLine(LOCTEXT("TraitInfoRequiredInterfaces", "Requires:").ToString());
+						for (const FTraitInterfaceUID& RequiredInterfaceUID : Trait->GetTraitRequiredInterfaces())
+						{
+							if (const ITraitInterface* RequiredInterface = FTraitInterfaceRegistry::Get().Find(RequiredInterfaceUID))
+							{
+								TraitInfoText.AppendLineFormat(EntryFormat, RequiredInterface->GetDisplayName());
+							}
+						}
+					}
+
+					return TraitInfoText.ToText();
+				}
+				return FText();
+			})
+			.Visibility(EVisibility::Visible)
+			+SOverlay::Slot()
+			[
+				SNew(SHorizontalBox)
+				+SHorizontalBox::Slot()
+				.AutoWidth()
+				.VAlign(VAlign_Center)
+				[
+					HandleWidget
+				]
+				+SHorizontalBox::Slot()
+				.FillWidth(1.0f)
+				.HAlign(HAlign_Right)
+				[
+					SNew(SSimpleComboButton)
+					.HasDownArrow(true)
+					.OnGetMenuContent_Lambda([this]()
+					{
+						constexpr bool bShouldCloseWindowAfterMenuSelection = true;
+						FMenuBuilder MenuBuilder(bShouldCloseWindowAfterMenuSelection, nullptr);
+
+						MenuBuilder.AddSubMenu(
+							LOCTEXT("AddLabel", "Add..."),
+							LOCTEXT("AddTooltip", "Add a new trait after this one"),
+							FNewMenuDelegate::CreateLambda([this](FMenuBuilder& InMenuBuilder)
+							{
+								FStructViewerModule& StructViewerModule = FModuleManager::LoadModuleChecked<FStructViewerModule>("StructViewer");
+
+								FStructViewerInitializationOptions Options;
+								Options.Mode = EStructViewerMode::StructPicker;
+								Options.NameTypeToDisplay = EStructViewerNameTypeToDisplay::DisplayName;
+								Options.StructFilter = MakeShared<FStructFilter>(Nodes, ETraitMode::Additive);
+
+								InMenuBuilder.AddWidget(
+									SNew(SBox)
+									.WidthOverride(300.0f)
+									.HeightOverride(300.0f)
+									[
+										StructViewerModule.CreateStructViewer(Options, FOnStructPicked::CreateLambda([this](const UScriptStruct* InStruct)
+										{
+											FSlateApplication::Get().DismissAllMenus();
+
+											UAnimNextController* Controller = Cast<UAnimNextController>(WeakController.Get());
+											if (Controller == nullptr)
+											{
+												return;
+											}
+
+											TSharedPtr<IPropertyHandle> PropertyHandle = WeakPropertyHandle.Pin();
+											if (!PropertyHandle.IsValid())
+											{
+												return;
+											}
+
+											TSharedPtr<IPropertyHandleStruct> StructHandle = PropertyHandle->AsStruct();
+											check(StructHandle.IsValid());
+
+											TSharedPtr<FStructOnScope> StructData = StructHandle->GetStructData();
+
+											TInstancedStruct<FAnimNextTraitSharedData> Defaults;
+											Defaults.InitializeAsScriptStruct(InStruct);
+											
+											FRigVMControllerCompileBracketScope CompileBracketScope(Controller);
+											Controller->OpenUndoBracket(LOCTEXT("AddTraitTransaction", "Add Trait").ToString());
+
+											for (TWeakObjectPtr<UAnimNextEdGraphNode> WeakNode : Nodes)
+											{
+												if (UAnimNextEdGraphNode* Node = WeakNode.Get())
+												{
+													TArray<URigVMPin*> TraitPins = Node->GetModelNode()->GetTraitPins();
+													int32 SourcePinIndex = INDEX_NONE;
+													for (URigVMPin* Pin : TraitPins)
+													{
+														TSharedPtr<FStructOnScope> TraitScope = Pin->GetTraitInstance();
+														if (TraitScope.IsValid()) 
+														{
+															const FRigVMTrait* VMTrait = (FRigVMTrait*)TraitScope->GetStructMemory();
+															if (VMTrait->GetTraitSharedDataStruct() == StructData->GetStruct())
+															{
+																SourcePinIndex = Pin->GetPinIndex();
+																break;
+															}
+														}
+													}
+
+													FName NewTrait = Controller->AddTraitStruct(CastChecked<URigVMUnitNode>(Node->GetModelNode()), Defaults, FMath::Min(SourcePinIndex + 1, Node->GetModelNode()->GetPins().Num()), true, true);
+													if (NewTrait != NAME_None)
+													{
+														FRigVMNodeLayout NodeLayout = Node->GetModelNode()->GetNodeLayout();
+														UUAFGraphNodeTemplate::AddDefaultTraitPinsToLayout(InStruct, NodeLayout);
+														Controller->SetNodeLayout(Node->GetModelNodeName(), NodeLayout, true, true);
+													}
+												}
+											}
+
+											Controller->CloseUndoBracket();
+										}))
+									], FText::GetEmpty(), true);
+							}));
+
+						MenuBuilder.AddSubMenu(
+							LOCTEXT("ReplaceLabel", "Replace..."),
+							LOCTEXT("ReplaceTooltip", "Replace this trait with another"),
+							FNewMenuDelegate::CreateLambda([this](FMenuBuilder& InMenuBuilder)
+							{
+								const FTraitRegistry& TraitRegistry = FTraitRegistry::Get();
+								const FTrait* Trait = TraitRegistry.Find(SharedDataStruct);
+								check(Trait);
+
+								FStructViewerModule& StructViewerModule = FModuleManager::LoadModuleChecked<FStructViewerModule>("StructViewer");
+
+								FStructViewerInitializationOptions Options;
+								Options.Mode = EStructViewerMode::StructPicker;
+								Options.NameTypeToDisplay = EStructViewerNameTypeToDisplay::DisplayName;
+								Options.StructFilter = MakeShared<FStructFilter>(Nodes, Trait->GetTraitMode());
+
+								InMenuBuilder.AddWidget(
+									SNew(SBox)
+									.WidthOverride(300.0f)
+									.HeightOverride(300.0f)
+									[
+										StructViewerModule.CreateStructViewer(Options, FOnStructPicked::CreateLambda([this](const UScriptStruct* InStruct)
+										{
+											FSlateApplication::Get().DismissAllMenus();
+
+											UAnimNextController* Controller = Cast<UAnimNextController>(WeakController.Get());
+											if (Controller == nullptr)
+											{
+												return;
+											}
+
+											TSharedPtr<IPropertyHandle> PropertyHandle = WeakPropertyHandle.Pin();
+											if (!PropertyHandle.IsValid())
+											{
+												return;
+											}
+
+											TSharedPtr<IPropertyHandleStruct> StructHandle = PropertyHandle->AsStruct();
+											check(StructHandle.IsValid());
+
+											TSharedPtr<FStructOnScope> StructData = StructHandle->GetStructData();
+
+											TInstancedStruct<FAnimNextTraitSharedData> Defaults;
+											Defaults.InitializeAsScriptStruct(InStruct);
+											
+											FRigVMControllerCompileBracketScope CompileBracketScope(Controller);
+											Controller->OpenUndoBracket(LOCTEXT("ReplaceTraitTransaction", "Replace Trait").ToString());
+
+											for (TWeakObjectPtr<UAnimNextEdGraphNode> WeakNode : Nodes)
+											{
+												if (UAnimNextEdGraphNode* Node = WeakNode.Get())
+												{
+													TArray<URigVMPin*> TraitPins = Node->GetModelNode()->GetTraitPins();
+													int32 SourcePinIndex = INDEX_NONE;
+													FName SourcePinName = NAME_None;
+													for (URigVMPin* Pin : TraitPins)
+													{
+														TSharedPtr<FStructOnScope> TraitScope = Pin->GetTraitInstance();
+														if (TraitScope.IsValid()) 
+														{
+															const FRigVMTrait* VMTrait = (FRigVMTrait*)TraitScope->GetStructMemory();
+															if (VMTrait->GetTraitSharedDataStruct() == StructData->GetStruct())
+															{
+																SourcePinIndex = Pin->GetPinIndex();
+																SourcePinName = Pin->GetFName();
+																break;
+															}
+														}
+													}
+
+													if (SourcePinIndex == INDEX_NONE || SourcePinName == NAME_None)
+													{
+														continue;
+													}
+
+													if (Controller->RemoveTraitByName(Node->GetModelNodeName(),SourcePinName, true, true))
+													{
+														FName NewTrait = Controller->AddTraitStruct(CastChecked<URigVMUnitNode>(Node->GetModelNode()), Defaults, SourcePinIndex, true, true);
+														if (NewTrait != NAME_None)
+														{
+															FRigVMNodeLayout NodeLayout = Node->GetModelNode()->GetNodeLayout();
+															UUAFGraphNodeTemplate::AddDefaultTraitPinsToLayout(InStruct, NodeLayout);
+															Controller->SetNodeLayout(Node->GetModelNodeName(), NodeLayout, true, true);
+														}
+													}
+												}
+											}
+
+											Controller->CloseUndoBracket();
+										}))
+									], FText::GetEmpty(), true);
+							}));
+
+						MenuBuilder.AddMenuEntry(
+							LOCTEXT("RemoveLabel", "Remove"),
+							LOCTEXT("RemoveTooltip", "Remove this trait"),
+							FSlateIcon(),
+							FExecuteAction::CreateLambda([this]()
+							{
+								FSlateApplication::Get().DismissAllMenus();
+
+								UAnimNextController* Controller = Cast<UAnimNextController>(WeakController.Get());
+								if (Controller == nullptr)
+								{
+									return;
+								}
+
+								TSharedPtr<IPropertyHandle> PropertyHandle = WeakPropertyHandle.Pin();
+								if (!PropertyHandle.IsValid())
+								{
+									return;
+								}
+
+								TSharedPtr<IPropertyHandleStruct> StructHandle = PropertyHandle->AsStruct();
+								check(StructHandle.IsValid());
+
+								TSharedPtr<FStructOnScope> StructData = StructHandle->GetStructData();
+								
+								FRigVMControllerCompileBracketScope CompileBracketScope(Controller);
+								Controller->OpenUndoBracket(LOCTEXT("RemoveTraitTransaction", "Remove Trait").ToString());
+
+								for (TWeakObjectPtr<UAnimNextEdGraphNode> WeakNode : Nodes)
+								{
+									if (UAnimNextEdGraphNode* Node = WeakNode.Get())
+									{
+										Controller->RemoveTraitByName(Node->GetModelNodeName(),StructData->GetStruct()->GetFName(), true, true);
+									}
+								}
+
+								Controller->CloseUndoBracket();
+							}));
+
+						return MenuBuilder.MakeWidget();
+					})
+				]
+			];
+
+		ChildSlot
+		[
+			Widget
+		];
+	}
+}
+
+TSharedRef<SWidget> FTraitStackEditor::CreateTraitHeaderWidget(IDetailCategoryBuilder& InCategory, IDetailPropertyRow& InPropertyRow, URigVMController* InController)
+{
+	return SNew(Private::STraitHeader, InCategory, InPropertyRow, InController);
 }
+
+TSharedRef<SWidget> FTraitStackEditor::CreateTraitStackHeaderWidget(IDetailCategoryBuilder& InCategory, URigVMController* InController)
+{
+	TArray<TWeakObjectPtr<UAnimNextEdGraphNode>> Nodes = InCategory.GetParentLayout().GetObjectsOfTypeBeingCustomized<UAnimNextEdGraphNode>();
+	
+	return
+		SNew(SHorizontalBox)
+		+SHorizontalBox::Slot()
+		.FillWidth(1.0f)
+		.HAlign(HAlign_Right)
+		.VAlign(VAlign_Center)
+		[
+			SNew(SPositiveActionButton)
+			.Text(LOCTEXT("AddTrait", "Add Trait"))
+			.ToolTipText(LOCTEXT("AddTraitTooltip", "Add a new trait to this node"))
+			.OnGetMenuContent_Lambda([Nodes = MoveTemp(Nodes), InController]()
+			{
+				ETraitMode Mode = ETraitMode::Additive;
+				bool bAreAllNodesEmpty = true;
+
+				for (TWeakObjectPtr<UAnimNextEdGraphNode> WeakNode : Nodes)
+				{
+					if (UAnimNextEdGraphNode* Node = WeakNode.Get())
+					{
+
+						bAreAllNodesEmpty &= (Node->GetModelNode()->GetTraitPins().Num() == 0);
+					}
+				}
+
+				if (bAreAllNodesEmpty)
+				{
+					Mode = ETraitMode::Base;
+				}
+
+				FStructViewerModule& StructViewerModule = FModuleManager::LoadModuleChecked<FStructViewerModule>("StructViewer");
+
+				FStructViewerInitializationOptions Options;
+				Options.Mode = EStructViewerMode::StructPicker;
+				Options.NameTypeToDisplay = EStructViewerNameTypeToDisplay::DisplayName;
+				Options.StructFilter = MakeShared<Private::FStructFilter>(Nodes, Mode);
+
+				return SNew(SBox)
+				.WidthOverride(300.0f)
+				.HeightOverride(300.0f)
+				[
+					StructViewerModule.CreateStructViewer(Options, FOnStructPicked::CreateLambda([InController, &Nodes](const UScriptStruct* InStruct)
+					{
+						FSlateApplication::Get().DismissAllMenus();
+
+						UAnimNextController* Controller = Cast<UAnimNextController>(InController);
+						if (Controller == nullptr)
+						{
+							return;
+						}
+
+						TInstancedStruct<FAnimNextTraitSharedData> Defaults;
+						Defaults.InitializeAsScriptStruct(InStruct);
+
+						FRigVMControllerCompileBracketScope CompileBracketScope(Controller);
+						Controller->OpenUndoBracket(LOCTEXT("AddTraitTransaction", "Add Trait").ToString());
+
+						for (TWeakObjectPtr<UAnimNextEdGraphNode> WeakNode : Nodes)
+						{
+							if (UAnimNextEdGraphNode* Node = WeakNode.Get())
+							{
+								FName NewTrait = Controller->AddTraitStruct(CastChecked<URigVMUnitNode>(Node->GetModelNode()), Defaults, Node->GetModelNode()->GetPins().Num(), true, true);
+								if (NewTrait != NAME_None)
+								{
+									FRigVMNodeLayout NodeLayout = Node->GetModelNode()->GetNodeLayout();
+									UUAFGraphNodeTemplate::AddDefaultTraitPinsToLayout(InStruct, NodeLayout);
+									Controller->SetNodeLayout(Node->GetModelNodeName(), NodeLayout, true, true);
+								}
+							}
+						}
+
+						Controller->CloseUndoBracket();
+					}))
+				];
+			})
+		];
+}
+
+}
+
+#undef LOCTEXT_NAMESPACE
\ No newline at end of file
diff --git a/Engine/Plugins/Experimental/UAF/UAFAnimGraph/Source/UAFAnimGraphEditor/Private/TraitStackEditor.h b/Engine/Plugins/Experimental/UAF/UAFAnimGraph/Source/UAFAnimGraphEditor/Private/TraitStackEditor.h
--- a/Engine/Plugins/Experimental/UAF/UAFAnimGraph/Source/UAFAnimGraphEditor/Private/TraitStackEditor.h
+++ b/Engine/Plugins/Experimental/UAF/UAFAnimGraph/Source/UAFAnimGraphEditor/Private/TraitStackEditor.h
@@ -4,6 +4,11 @@
 
 #include "ITraitStackEditor.h"
 
+class IDetailPropertyRow;
+class IDetailCategoryBuilder;
+class SWidget;
+class URigVMController;
+
 namespace UE::UAF::Editor
 {
 
@@ -12,6 +17,8 @@ class FTraitStackEditor : public ITraitStackEditor
 {
 	// ITraitStackEditor interface
 	virtual void SetTraitData(const TSharedRef<UE::Workspace::IWorkspaceEditor>& InWorkspaceEditor, const FTraitStackData& InTraitStackData) override;
+	virtual TSharedRef<SWidget> CreateTraitHeaderWidget(IDetailCategoryBuilder& InCategory, IDetailPropertyRow& InPropertyRow, URigVMController* InController) override;
+	virtual TSharedRef<SWidget> CreateTraitStackHeaderWidget(IDetailCategoryBuilder& InCategory, URigVMController* InController) override;
 };
 
 }
\ No newline at end of file
diff --git a/Engine/Plugins/Experimental/UAF/UAFAnimGraph/Source/UAFAnimGraphUncookedOnly/Internal/AnimNextController.h b/Engine/Plugins/Experimental/UAF/UAFAnimGraph/Source/UAFAnimGraphUncookedOnly/Internal/AnimNextController.h
--- a/Engine/Plugins/Experimental/UAF/UAFAnimGraph/Source/UAFAnimGraphUncookedOnly/Internal/AnimNextController.h
+++ b/Engine/Plugins/Experimental/UAF/UAFAnimGraph/Source/UAFAnimGraphUncookedOnly/Internal/AnimNextController.h
@@ -26,7 +26,7 @@ public:
 	// Adds a new Trait to the Stack, with default struct values
 	// Returns Trait Instance Name (or NAME_None on failure)
 	UFUNCTION(BlueprintCallable, Category = "UAF|Controller")
-	FName AddTraitStruct(URigVMUnitNode* InNode, const TInstancedStruct<FAnimNextTraitSharedData>& InTraitDefaults, int32 InTraitIndex, bool bSetupUndoRedo, bool bPrintPythonCommand);
+	UAFANIMGRAPHUNCOOKEDONLY_API FName AddTraitStruct(URigVMUnitNode* InNode, const TInstancedStruct<FAnimNextTraitSharedData>& InTraitDefaults, int32 InTraitIndex, bool bSetupUndoRedo, bool bPrintPythonCommand);
 
 	// Adds a new Trait to the Stack, with default struct values
 	// Returns Trait Instance Name (or NAME_None on failure)
diff --git a/Engine/Plugins/Experimental/UAF/UAFAnimGraph/Source/UAFAnimGraphUncookedOnly/Private/AnimNextController.cpp b/Engine/Plugins/Experimental/UAF/UAFAnimGraph/Source/UAFAnimGraphUncookedOnly/Private/AnimNextController.cpp
--- a/Engine/Plugins/Experimental/UAF/UAFAnimGraph/Source/UAFAnimGraphUncookedOnly/Private/AnimNextController.cpp
+++ b/Engine/Plugins/Experimental/UAF/UAFAnimGraph/Source/UAFAnimGraphUncookedOnly/Private/AnimNextController.cpp
@@ -155,7 +155,8 @@ FName UAnimNextController::AddTraitByName(FName InNodeName, FName InNewTraitType
 	// Avoid multiple VM recompilations for internal operations
 	FRigVMControllerCompileBracketScope CompileScope(this);
 
-	const FName TraitName = InNewTraitTypeName;
+	// Must use struct name for name to allow for consistent pathing in node layouts
+	const FName TraitName = ScriptStruct->GetFName();
 
 	return AddTrait(InNodeName, *CppDecoratorStruct->GetPathName(), TraitName, DefaultValue, InPinIndex, bSetupUndoRedo, bPrintPythonCommand);
 }
diff --git a/Engine/Plugins/Experimental/UAF/UAFAnimGraph/Source/UAFAnimGraphUncookedOnly/Private/AnimNextTraitStackUnitNode.cpp b/Engine/Plugins/Experimental/UAF/UAFAnimGraph/Source/UAFAnimGraphUncookedOnly/Private/AnimNextTraitStackUnitNode.cpp
--- a/Engine/Plugins/Experimental/UAF/UAFAnimGraph/Source/UAFAnimGraphUncookedOnly/Private/AnimNextTraitStackUnitNode.cpp
+++ b/Engine/Plugins/Experimental/UAF/UAFAnimGraph/Source/UAFAnimGraphUncookedOnly/Private/AnimNextTraitStackUnitNode.cpp
@@ -57,7 +57,7 @@ const FSlateBrush* UAnimNextTraitStackUnitNode::GetDefaultNodeIconBrush() const
 		return Template->GetDefaultObject<UUAFGraphNodeTemplate>()->GetIconBrush();
 	}
 
-	return FStyleDefaults::GetNoBrush();
+	return FSlateIcon("UAFAnimGraphUncookedOnlyStyle", "NodeTemplate.DefaultIcon").GetIcon();
 }
 
 void UAnimNextTraitStackUnitNode::HandlePinDefaultValueChanged(UAnimNextController* InController, URigVMPin* InPin)
diff --git a/Engine/Plugins/Experimental/UAF/UAFAnimGraph/Source/UAFAnimGraphUncookedOnly/Private/Graph/AnimNextAnimationGraph_EditorData.cpp b/Engine/Plugins/Experimental/UAF/UAFAnimGraph/Source/UAFAnimGraphUncookedOnly/Private/Graph/AnimNextAnimationGraph_EditorData.cpp
--- a/Engine/Plugins/Experimental/UAF/UAFAnimGraph/Source/UAFAnimGraphUncookedOnly/Private/Graph/AnimNextAnimationGraph_EditorData.cpp
+++ b/Engine/Plugins/Experimental/UAF/UAFAnimGraph/Source/UAFAnimGraphUncookedOnly/Private/Graph/AnimNextAnimationGraph_EditorData.cpp
@@ -32,6 +32,7 @@
 #include "Entries/AnimNextRigVMAssetEntry.h"
 #include "Graph/RigUnit_AnimNextGraphEvaluator.h"
 #include "Graph/RigUnit_AnimNextGraphRoot.h"
+#include "TraitCore/TraitInterfaceRegistry.h"
 
 #include UE_INLINE_GENERATED_CPP_BY_NAME(AnimNextAnimationGraph_EditorData)
 
@@ -325,11 +326,128 @@ namespace UE::UAF::UncookedOnly::Private
 		}
 	}
 
+	bool ValidateTraitStack(const FRigVMCompileSettings& InSettings, const FTraitStackMapping& InMapping, URigVMUnitNode* InVMUnitNode)
+	{
+		bool bResult = true; 
+		for (int32 TraitIndex = 0; TraitIndex < InMapping.TraitEntries.Num(); ++TraitIndex)
+		{
+			const FTraitEntryMapping& Entry = InMapping.TraitEntries[TraitIndex];
+			switch (Entry.Trait->GetTraitMode())
+			{
+			case ETraitMode::Base:
+				{
+					if (Entry.Trait->GetTraitUID() == FTraitUID())
+					{
+						InSettings.ASTSettings.Report(EMessageSeverity::Error, InVMUnitNode, LOCTEXT("TraitStatusInStack_InvalidBaseTrait", "@@: Base trait is invalid. Please select a new base trait.").ToString());
+						bResult = false;
+					}
+					if (Entry.Trait->GetTraitUID() != InMapping.TraitEntries[0].Trait->GetTraitUID())
+					{
+						InSettings.ASTSettings.Report(EMessageSeverity::Error, InVMUnitNode, LOCTEXT("TraitStatusInStack_BaseNotAtTop", "@@: Base traits should be the first trait in the stack.").ToString());
+						bResult = false;
+					}
+					break;
+				}
+			case ETraitMode::Additive:
+				{
+					if (Entry.Trait->GetTraitUID() == FTraitUID())
+					{
+						InSettings.ASTSettings.Report(EMessageSeverity::Error, InVMUnitNode, LOCTEXT("TraitStatusInStack_InvalidAdditiveTrait", "@@: Additive trait is invalid. Please fix the stack.").ToString());
+						bResult = false;
+					}
+					if (Entry.Trait->GetTraitUID() == InMapping.TraitEntries[0].Trait->GetTraitUID())
+					{
+						InSettings.ASTSettings.Report(EMessageSeverity::Error, InVMUnitNode, LOCTEXT("TraitStatusInStack_AdditiveAtTop", "@@: Additive traits cannot be at the top of the stack.").ToString());
+						bResult = false;
+					}
+					break;
+				}
+			default:
+				{
+					InSettings.ASTSettings.Report(EMessageSeverity::Error, InVMUnitNode, LOCTEXT("TraitStatusInStack_InvalidTraitData", "@@: Trait is invalid, please correct the stack.").ToString());
+					bResult = false;
+					break;
+				}
+			}
+
+			const TConstArrayView<FTraitInterfaceUID> RequiredInterfaces = Entry.Trait->GetTraitRequiredInterfaces();
+			const int32 NumRequired = RequiredInterfaces.Num();
+			if (NumRequired > 0)
+			{
+				for (int32 RequiredIndex = 0; RequiredIndex < NumRequired; RequiredIndex++)
+				{
+					const FTraitInterfaceUID& RequiredInterface = RequiredInterfaces[RequiredIndex];
+
+					bool bFound = false;
+					const int32 StartIndex = (Entry.Trait->GetTraitMode() == ETraitMode::Base) 
+						? InMapping.TraitEntries.Num() - 1 // Base traits scan all stack to find a valid interface
+						: TraitIndex; // Additive traits start search from current trait to enable traits that inherit from a trait with a required interface that is implemented in the derived class
+					
+					for (int32 SearchStartIndex = StartIndex; SearchStartIndex >= 0 && !bFound; SearchStartIndex--)
+					{
+						const FTraitEntryMapping& ParentEntry = InMapping.TraitEntries[SearchStartIndex];
+
+						const TConstArrayView<FTraitInterfaceUID> ParentImplementedInterfaces = ParentEntry.Trait->GetTraitInterfaces();
+						for (const FTraitInterfaceUID& ParentImplementedInterface : ParentImplementedInterfaces)
+						{
+							if (ParentImplementedInterface == RequiredInterface)
+							{
+								bFound = true;
+								break;
+							}
+						}
+					}
+					if (!bFound)
+					{
+						if (const ITraitInterface* TraitInterface = FTraitInterfaceRegistry::Get().Find(RequiredInterface))
+						{
+							const FText InterfaceName(TraitInterface->GetDisplayName());
+							const FText MissingError = FText::Format(LOCTEXT("TraitStatusInStack_MissingInterface", "@@: Trait '{0}' requires a parent implementing interface '{1}'"), Entry.Trait->GetTraitSharedDataStruct()->GetDisplayNameText(), InterfaceName);
+							InSettings.ASTSettings.Report(EMessageSeverity::Warning, InVMUnitNode, MissingError.ToString());
+						}
+						else
+						{
+							const FText UID = FText::FromString(FString::Printf(TEXT("0x%x"), RequiredInterface.GetUID()));
+							const FText UnknownError = FText::Format(LOCTEXT("TraitStatusInStack_UnknownInterface", "@@: Trait '{0}' requires a parent implementing unknown interface with ID '{1}'"), Entry.Trait->GetTraitSharedDataStruct()->GetDisplayNameText(), UID);
+							InSettings.ASTSettings.Report(EMessageSeverity::Warning, InVMUnitNode, UnknownError.ToString());
+						}
+					}
+				}
+			}
+		}
+		return bResult;
+	}
+
 	FTraitGraph CollectGraphInfo(const FRigVMCompileSettings& InSettings, const UAnimNextAnimationGraph* InAnimationGraph, const URigVMGraph* VMGraph, URigVMController* VMController)
 	{
 		const TArray<URigVMNode*>& VMNodes = VMGraph->GetNodes();
 		URigVMUnitNode* VMRootNode = FindRootNode(VMNodes);
 
+		// Pre-validate all nodes, connected or otherwise, to discover invalid trait stacks
+		TMap<URigVMNode*, FTraitStackMapping> ValidNodeMappings;
+		ValidNodeMappings.Reserve(VMNodes.Num());
+		for (URigVMNode* VMNode : VMNodes)
+		{
+			if (const URigVMUnitNode* VMUnitNode = Cast<URigVMUnitNode>(VMNode))
+			{
+				const UScriptStruct* ScriptStruct = VMUnitNode->GetScriptStruct();
+				if (ScriptStruct == FRigUnit_AnimNextTraitStack::StaticStruct())
+				{
+					FTraitStackMapping Mapping(VMNode);
+					ForEachTraitInStack(VMNode,
+						[&Mapping](const URigVMNode* DecoratorStackNode, const URigVMPin* DecoratorPin, const FTrait* Trait)
+						{
+							Mapping.TraitEntries.Add(FTraitEntryMapping(DecoratorStackNode, DecoratorPin, Trait));
+						});
+
+					if(ValidateTraitStack(InSettings, Mapping, const_cast<URigVMUnitNode*>(VMUnitNode)))
+					{
+						ValidNodeMappings.Add(VMNode, MoveTemp(Mapping));
+					}
+				}
+			}
+		}
+
 		if (VMRootNode == nullptr)
 		{
 			// Root node wasn't found, add it, we'll need it to compile
@@ -349,31 +467,9 @@ namespace UE::UAF::UncookedOnly::Private
 			const URigVMNode* VMNode = NodesToVisit[0];
 			NodesToVisit.RemoveAt(0);
 
-			if (const URigVMUnitNode* VMUnitNode = Cast<URigVMUnitNode>(VMNode))
+			if (FTraitStackMapping* MappingPtr = ValidNodeMappings.Find(VMNode))
 			{
-				const UScriptStruct* ScriptStruct = VMUnitNode->GetScriptStruct();
-				if (ScriptStruct == FRigUnit_AnimNextTraitStack::StaticStruct())
-				{
-					FTraitStackMapping Mapping(VMNode);
-
-					bool bHasBaseTrait = false;
-					ForEachTraitInStack(VMNode,
-						[&Mapping, &bHasBaseTrait](const URigVMNode* DecoratorStackNode, const URigVMPin* DecoratorPin, const FTrait* Trait)
-						{
-							bHasBaseTrait |= Trait->GetTraitMode() == ETraitMode::Base;
-							Mapping.TraitEntries.Add(FTraitEntryMapping(DecoratorStackNode, DecoratorPin, Trait));
-						});
-
-					if(!bHasBaseTrait)
-					{
-						// Must have at least one base trait
-						InSettings.ASTSettings.Reportf(EMessageSeverity::Error, const_cast<URigVMUnitNode*>(VMUnitNode), TEXT("No base trait supplied for @@"));
-					}
-					else
-					{
-						TraitGraph.TraitStackNodes.Add(MoveTemp(Mapping));
-					}
-				}
+				TraitGraph.TraitStackNodes.Add(MoveTemp(*MappingPtr));
 			}
 
 			const TArray<URigVMNode*> SourceNodes = VMNode->GetLinkedSourceNodes();
diff --git a/Engine/Plugins/Experimental/UAF/UAFAnimGraph/Source/UAFAnimGraphUncookedOnly/Private/Templates/UAFGraphNodeTemplate.cpp b/Engine/Plugins/Experimental/UAF/UAFAnimGraph/Source/UAFAnimGraphUncookedOnly/Private/Templates/UAFGraphNodeTemplate.cpp
--- a/Engine/Plugins/Experimental/UAF/UAFAnimGraph/Source/UAFAnimGraphUncookedOnly/Private/Templates/UAFGraphNodeTemplate.cpp
+++ b/Engine/Plugins/Experimental/UAF/UAFAnimGraph/Source/UAFAnimGraphUncookedOnly/Private/Templates/UAFGraphNodeTemplate.cpp
@@ -78,25 +78,40 @@ void UUAFGraphNodeTemplate::HandleAssetDropped_Implementation(UAnimNextControlle
 FRigVMNodeLayout UUAFGraphNodeTemplate::GetDefaultCategoryLayout() const
 {
 	FRigVMNodeLayout Layout;
-	FRigVMPinCategory& NewCategory = Layout.Categories.AddDefaulted_GetRef();
+	TMap<FString, FString> DisplayNames;
+	TMap<FString, int32> PinIndexInCategory;
+	FRigVMPinCategory NewCategory;
 	NewCategory.Path = FRigVMPinCategory::GetDefaultCategoryName();
 	for (const TInstancedStruct<FAnimNextTraitSharedData>& Trait : Traits)
 	{
 		if (const UScriptStruct* Struct = Trait.GetScriptStruct())
 		{
-			for (TFieldIterator<FProperty> It(Struct); It; ++It)
+			int32 PinIndex = 0;
+			for (TFieldIterator<FProperty> It(Struct); It; ++It, ++PinIndex)
 			{
-				TStringBuilder<128> PinPathBuilder;
-				PinPathBuilder.Append(Struct->GetName());
-				PinPathBuilder.AppendChar(TEXT('.'));
-				PinPathBuilder.Append(It->GetName());
-				FString PinPath = PinPathBuilder.ToString();
-				NewCategory.Elements.Add(PinPath);
-				Layout.DisplayNames.Add(PinPath, It->GetDisplayNameText().ToString());
+				if (!It->HasMetaData("PinHiddenByDefault") &&
+					!It->HasMetaData("Hidden"))
+				{
+					TStringBuilder<128> PinPathBuilder;
+					PinPathBuilder.Append(Struct->GetName());
+					PinPathBuilder.AppendChar(TEXT('.'));
+					PinPathBuilder.Append(It->GetName());
+					FString PinPath = PinPathBuilder.ToString();
+					NewCategory.Elements.Add(PinPath);
+					DisplayNames.Add(PinPath, It->GetDisplayNameText().ToString());
+					PinIndexInCategory.Add(PinPath, PinIndex);
+				}
 			}
 		}
 	}
 
+	if (NewCategory.Elements.Num() > 0)
+	{
+		Layout.Categories.Add(MoveTemp(NewCategory));
+		Layout.PinIndexInCategory = MoveTemp(PinIndexInCategory);
+		Layout.DisplayNames = MoveTemp(DisplayNames);
+	}
+
 	return Layout;
 }
 
@@ -108,17 +123,33 @@ FRigVMNodeLayout UUAFGraphNodeTemplate::GetPerTraitCategoriesLayout() const
 	{
 		if (const UScriptStruct* Struct = Trait.GetScriptStruct())
 		{
-			FRigVMPinCategory& NewCategory = Layout.Categories.AddDefaulted_GetRef();
+			TMap<FString, FString> DisplayNames;
+			TMap<FString, int32> PinIndexInCategory;
+			FRigVMPinCategory NewCategory;
 			NewCategory.Path = Struct->GetDisplayNameText().ToString();
-			for (TFieldIterator<FProperty> It(Struct); It; ++It)
+
+			int32 PinIndex = 0;
+			for (TFieldIterator<FProperty> It(Struct); It; ++It, ++PinIndex)
 			{
-				TStringBuilder<128> PinPathBuilder;
-				PinPathBuilder.Append(Struct->GetName());
-				PinPathBuilder.AppendChar(TEXT('.'));
-				PinPathBuilder.Append(It->GetName());
-				FString PinPath = PinPathBuilder.ToString();
-				NewCategory.Elements.Add(PinPath);
-				Layout.DisplayNames.Add(PinPath, It->GetDisplayNameText().ToString());
+				if (!It->HasMetaData("PinHiddenByDefault") &&
+					!It->HasMetaData("Hidden"))
+				{
+					TStringBuilder<128> PinPathBuilder;
+					PinPathBuilder.Append(Struct->GetName());
+					PinPathBuilder.AppendChar(TEXT('.'));
+					PinPathBuilder.Append(It->GetName());
+					FString PinPath = PinPathBuilder.ToString();
+					NewCategory.Elements.Add(PinPath);
+					DisplayNames.Add(PinPath, It->GetDisplayNameText().ToString());
+					PinIndexInCategory.Add(PinPath, PinIndex);
+				}
+			}
+
+			if (NewCategory.Elements.Num() > 0)
+			{
+				Layout.Categories.Add(MoveTemp(NewCategory));
+				Layout.PinIndexInCategory = MoveTemp(PinIndexInCategory);
+				Layout.DisplayNames = MoveTemp(DisplayNames);
 			}
 		}
 	}
@@ -126,6 +157,43 @@ FRigVMNodeLayout UUAFGraphNodeTemplate::GetPerTraitCategoriesLayout() const
 	return Layout;
 }
 
+void UUAFGraphNodeTemplate::AddDefaultTraitPinsToLayout(const UScriptStruct* InStruct, FRigVMNodeLayout& InOutLayout)
+{
+	if (!InStruct->IsChildOf<FAnimNextTraitSharedData>())
+	{
+		return;
+	}
+
+	TMap<FString, FString> DisplayNames;
+	TMap<FString, int32> PinIndexInCategory;
+	FRigVMPinCategory NewCategory;
+	NewCategory.Path = InStruct->GetDisplayNameText().ToString();
+
+	int32 PinIndex = 0;
+	for (TFieldIterator<FProperty> It(InStruct); It; ++It, ++PinIndex)
+	{
+		if (!It->HasMetaData("PinHiddenByDefault") &&
+			!It->HasMetaData("Hidden"))
+		{
+			TStringBuilder<128> PinPathBuilder;
+			PinPathBuilder.Append(InStruct->GetName());
+			PinPathBuilder.AppendChar(TEXT('.'));
+			PinPathBuilder.Append(It->GetName());
+			FString PinPath = PinPathBuilder.ToString();
+			NewCategory.Elements.Add(PinPath);
+			DisplayNames.Add(PinPath, It->GetDisplayNameText().ToString());
+			PinIndexInCategory.Add(PinPath, PinIndex);
+		}
+	}
+
+	if (NewCategory.Elements.Num() > 0)
+	{
+		InOutLayout.Categories.Add(MoveTemp(NewCategory));
+		InOutLayout.PinIndexInCategory.Append(PinIndexInCategory);
+		InOutLayout.DisplayNames.Append(DisplayNames);
+	}
+}
+
 void UUAFGraphNodeTemplate::SetDisplayNameForPinInLayout(FString PinPath, FString PinDisplayName, FRigVMNodeLayout& Layout)
 {
 	Layout.DisplayNames.Add(PinPath, PinDisplayName);
diff --git a/Engine/Plugins/Experimental/UAF/UAFAnimGraph/Source/UAFAnimGraphUncookedOnly/Private/Templates/UAFGraphNodeTemplate_BlendSpacePlayer.h b/Engine/Plugins/Experimental/UAF/UAFAnimGraph/Source/UAFAnimGraphUncookedOnly/Private/Templates/UAFGraphNodeTemplate_BlendSpacePlayer.h
--- a/Engine/Plugins/Experimental/UAF/UAFAnimGraph/Source/UAFAnimGraphUncookedOnly/Private/Templates/UAFGraphNodeTemplate_BlendSpacePlayer.h
+++ b/Engine/Plugins/Experimental/UAF/UAFAnimGraph/Source/UAFAnimGraphUncookedOnly/Private/Templates/UAFGraphNodeTemplate_BlendSpacePlayer.h
@@ -27,8 +27,7 @@ class UUAFGraphNodeTemplate_BlendSpacePlayer : public UUAFGraphNodeTemplate
 		Icon = *FSlateIconFinder::FindIconForClass(UBlendSpace::StaticClass()).GetIcon();
 		Traits =
 		{
-			TInstancedStruct<UE::UAF::FBlendSpacePlayerData>::Make(),
-			TInstancedStruct<UE::UAF::FNotifyDispatcherData>::Make()
+			TInstancedStruct<UE::UAF::FBlendSpacePlayerData>::Make()
 		};
 		DragDropAssetTypes.Add(UBlendSpace::StaticClass());
 		SetCategoryForPinsInLayout(
diff --git a/Engine/Plugins/Experimental/UAF/UAFAnimGraph/Source/UAFAnimGraphUncookedOnly/Public/Templates/UAFGraphNodeTemplate.h b/Engine/Plugins/Experimental/UAF/UAFAnimGraph/Source/UAFAnimGraphUncookedOnly/Public/Templates/UAFGraphNodeTemplate.h
--- a/Engine/Plugins/Experimental/UAF/UAFAnimGraph/Source/UAFAnimGraphUncookedOnly/Public/Templates/UAFGraphNodeTemplate.h
+++ b/Engine/Plugins/Experimental/UAF/UAFAnimGraph/Source/UAFAnimGraphUncookedOnly/Public/Templates/UAFGraphNodeTemplate.h
@@ -154,6 +154,9 @@ public:
 	// Utility function used to refresh an existing node from this templates
 	UE_API void RefreshNodeFromTemplate(UAnimNextController* InController, UAnimNextTraitStackUnitNode* InNode) const;
 
+	// Utility function used to add all of a trait's pins to a category, named according to their display name
+	UE_API static void AddDefaultTraitPinsToLayout(const UScriptStruct* Struct, FRigVMNodeLayout& InOutLayout);
+
 protected:
 	friend UE::UAF::FGraphNodeTemplateRegistry;
 	friend UE::UAF::Editor::FGraphNodeTemplateDetails;
diff --git a/Engine/Plugins/Runtime/RigVM/Source/RigVMDeveloper/Private/RigVMModel/RigVMController.cpp b/Engine/Plugins/Runtime/RigVM/Source/RigVMDeveloper/Private/RigVMModel/RigVMController.cpp
--- a/Engine/Plugins/Runtime/RigVM/Source/RigVMDeveloper/Private/RigVMModel/RigVMController.cpp
+++ b/Engine/Plugins/Runtime/RigVM/Source/RigVMDeveloper/Private/RigVMModel/RigVMController.cpp
@@ -8505,7 +8505,7 @@ bool URigVMController::AddEmptyPinCategory(const URigVMNode* InNode, const FStri
 			// apply this to all pins - it's going to filter invalid pins anyway
 			for(URigVMPin* RootPin : InNode->Pins)
 			{
-				(void)SetPinCategory(RootPin, PinCategoryNameString, bSetupUndoRedo);
+				(void)SetPinCategory(RootPin, PinCategoryNameString, INDEX_NONE, bSetupUndoRedo);
 			}
 		}
 	}
@@ -8548,7 +8548,7 @@ bool URigVMController::SetPinCategory(const FString& InPinPath, const FString& I
 		return false;
 	}
 
-	const bool bSuccess = SetPinCategory(Pin, InCategory, bSetupUndoRedo);
+	const bool bSuccess = SetPinCategory(Pin, InCategory, INDEX_NONE, bSetupUndoRedo);
 	if (bSuccess && bPrintPythonCommand)
 	{
 		const FString GraphName = GetSchema()->GetSanitizedGraphName(GetGraph()->GetGraphName());
@@ -8568,7 +8568,7 @@ bool URigVMController::ClearPinCategory(const FString& InPinPath, bool bSetupUnd
 	return SetPinCategory(InPinPath, FString(), bSetupUndoRedo, bPrintPythonCommand);
 }
 
-bool URigVMController::SetPinCategory(URigVMPin* InPin, const FString& InCategory, bool bSetupUndoRedo)
+bool URigVMController::SetPinCategory(URigVMPin* InPin, const FString& InCategory, int32 InIndexInCategory, bool bSetupUndoRedo)
 {
 	if(!IsValidPinForGraph(InPin))
 	{
@@ -8620,7 +8620,7 @@ bool URigVMController::SetPinCategory(URigVMPin* InPin, const FString& InCategor
 
 	if(NewCategory.Equals(InPin->UserDefinedCategory, ESearchCase::CaseSensitive))
 	{
-		return false;
+		return true;
 	}
 
 	// we only allow one category per pin path - so we need to check all parents and children for set categories
@@ -8719,14 +8719,28 @@ bool URigVMController::SetPinCategory(URigVMPin* InPin, const FString& InCategor
 		}
 	}
 
-	int32 IndexInCategory = 0;
-	if(NewCategory.IsEmpty())
+	int32 IndexInCategory = INDEX_NONE;
+	if(InIndexInCategory == INDEX_NONE)
 	{
-		IndexInCategory = INDEX_NONE;
+		IndexInCategory = InPin->GetNode()->GetPinsForCategory(NewCategory).Num();
 	}
-	else
+	else if (!NewCategory.IsEmpty())
 	{
-		IndexInCategory = InPin->GetNode()->GetPinsForCategory(NewCategory).Num();
+		const TArray<URigVMPin*> PinsInNewCategory = InPin->GetNode()->GetPinsForCategory(NewCategory);
+		IndexInCategory = FMath::Min(InIndexInCategory, PinsInNewCategory.Num());
+
+		// move all pins in the new category unless we are adding to the end
+		if (IndexInCategory != PinsInNewCategory.Num())
+		{
+			for(URigVMPin* OtherPin : PinsInNewCategory)
+			{
+				if(OtherPin != InPin && OtherPin->IndexInCategory != INDEX_NONE && OtherPin->IndexInCategory >= IndexInCategory)
+				{
+					OtherPin->IndexInCategory++;
+					Notify(ERigVMGraphNotifType::PinCategoryChanged, OtherPin);
+				}
+			}
+		}
 	}
 
 	InPin->UserDefinedCategory = NewCategory;
@@ -8886,7 +8900,7 @@ bool URigVMController::RemovePinCategory(const URigVMNode* InNode, const FString
 	for(URigVMPin* PinInCategory : PinsForCategory)
 	{
 		const TGuardValue<bool> GuardRefreshFunctions(bSuspendRefreshingFunctionReferences, true); 
-		if(SetPinCategory(PinInCategory, FString(), bSetupUndoRedo))
+		if(SetPinCategory(PinInCategory, FString(), INDEX_NONE, bSetupUndoRedo))
 		{
 			NumPinsAffected++;
 		}
@@ -9062,7 +9076,7 @@ bool URigVMController::RenamePinCategory(const URigVMNode* InNode, const FString
 			// store the original pin category name as it will be needed later at the final set
 			PinsCategoriesToRename.Add(CategoryToRename);
 
-			(void)SetPinCategory(PinToUpdate, FString(), bSetupUndoRedo);
+			(void)SetPinCategory(PinToUpdate, FString(), INDEX_NONE, bSetupUndoRedo);
 		}
 	}
 
@@ -9096,12 +9110,12 @@ bool URigVMController::RenamePinCategory(const URigVMNode* InNode, const FString
 			const FString& CategoryToRename = PinsCategoriesToRename[PinIndex];
 			if(CategoryToRename.Equals(InOldPinCategory, ESearchCase::CaseSensitive))
 			{
-				(void)SetPinCategory(PinToUpdate, InNewPinCategory, bSetupUndoRedo);
+				(void)SetPinCategory(PinToUpdate, InNewPinCategory, INDEX_NONE, bSetupUndoRedo);
 			}
 			else if(CategoryToRename.StartsWith(OldPrefix, ESearchCase::CaseSensitive))
 			{
 				const FString NewCategory = NewPrefix + CategoryToRename.RightChop(OldPrefix.Len());
-				(void)SetPinCategory(PinToUpdate, NewCategory, bSetupUndoRedo);
+				(void)SetPinCategory(PinToUpdate, NewCategory, INDEX_NONE, bSetupUndoRedo);
 			}
 		}
 	}
@@ -9412,7 +9426,7 @@ bool URigVMController::SetPinIndexInCategory(URigVMPin* InPin, int32 InIndexInCa
 
 	if(InPin->GetIndexInCategory() == InIndexInCategory)
 	{
-		return false;
+		return true;
 	}
 
 	FRigVMControllerCompileBracketScope CompileScope(this);
diff --git a/Engine/Plugins/Runtime/RigVM/Source/RigVMDeveloper/Private/RigVMModel/RigVMControllerActions.cpp b/Engine/Plugins/Runtime/RigVM/Source/RigVMDeveloper/Private/RigVMModel/RigVMControllerActions.cpp
--- a/Engine/Plugins/Runtime/RigVM/Source/RigVMDeveloper/Private/RigVMModel/RigVMControllerActions.cpp
+++ b/Engine/Plugins/Runtime/RigVM/Source/RigVMDeveloper/Private/RigVMModel/RigVMControllerActions.cpp
@@ -1150,6 +1150,7 @@ FRigVMSetPinCategoryAction::FRigVMSetPinCategoryAction()
 : FRigVMBaseAction(nullptr)
 , PinPath()
 , OldCategory()
+, OldPinIndex(INDEX_NONE)
 , NewCategory()
 {
 }
@@ -1158,6 +1159,7 @@ FRigVMSetPinCategoryAction::FRigVMSetPinCategoryAction(URigVMController* InContr
 : FRigVMBaseAction(InController)
 , PinPath(InPin->GetPinPath())
 , OldCategory(InPin->UserDefinedCategory)
+, OldPinIndex(InPin->GetIndexInCategory())
 , NewCategory(InNewCategory)
 {
 }
@@ -1168,7 +1170,14 @@ bool FRigVMSetPinCategoryAction::Undo()
 	{
 		return false;
 	}
-	return GetController()->SetPinCategory(PinPath, OldCategory, false);
+	if(const URigVMGraph* Graph = GetController()->GetGraph())
+	{
+		if(URigVMPin* Pin = Graph->FindPin(PinPath))
+		{
+			return GetController()->SetPinCategory(Pin, OldCategory, OldPinIndex, false);
+		}
+	}
+	return false;
 }
 
 bool FRigVMSetPinCategoryAction::Redo()
diff --git a/Engine/Plugins/Runtime/RigVM/Source/RigVMDeveloper/Public/RigVMModel/RigVMController.h b/Engine/Plugins/Runtime/RigVM/Source/RigVMDeveloper/Public/RigVMModel/RigVMController.h
--- a/Engine/Plugins/Runtime/RigVM/Source/RigVMDeveloper/Public/RigVMModel/RigVMController.h
+++ b/Engine/Plugins/Runtime/RigVM/Source/RigVMDeveloper/Public/RigVMModel/RigVMController.h
@@ -1368,7 +1368,7 @@ private:
 	UE_API bool SetPinIsWatched(URigVMPin* InPin, bool bIsWatched, bool bSetupUndoRedo);
 	UE_API bool SetPinDisplayName(URigVMPin* InPin, const FString& InDisplayName, bool bSetupUndoRedo);
 	UE_API bool AddEmptyPinCategory(const URigVMNode* InNode, const FString& InPinCategory, bool bSetupUndoRedo);
-	UE_API bool SetPinCategory(URigVMPin* InPin, const FString& InCategory, bool bSetupUndoRedo);
+	UE_API bool SetPinCategory(URigVMPin* InPin, const FString& InCategory, int32 InIndexInCategory, bool bSetupUndoRedo);
 	UE_API bool RemovePinCategory(const URigVMNode* InNode, const FString& InPinCategory, bool bSetupUndoRedo);
 	UE_API bool RenamePinCategory(const URigVMNode* InNode, const FString& InOldPinCategory, const FString& InNewPinCategory, bool bSetupUndoRedo);
 	UE_API bool SetPinCategoryIndex(const URigVMNode* InNode, const FString& InPinCategory, int32 InNewIndex, bool bSetupUndoRedo);
@@ -1702,6 +1702,7 @@ private:
 	friend struct FRigVMInjectNodeIntoPinAction;
 	friend struct FRigVMEjectNodeFromPinAction;
 	friend struct FRigVMChangeNodePinCategoriesAction;
+	friend struct FRigVMSetPinCategoryAction;
 	friend class FRigVMParserAST;
 	friend class FRigVMControllerCompileBracketScope;
 	friend struct FRigVMPinInfoArray;
diff --git a/Engine/Plugins/Runtime/RigVM/Source/RigVMDeveloper/Public/RigVMModel/RigVMControllerActions.h b/Engine/Plugins/Runtime/RigVM/Source/RigVMDeveloper/Public/RigVMModel/RigVMControllerActions.h
--- a/Engine/Plugins/Runtime/RigVM/Source/RigVMDeveloper/Public/RigVMModel/RigVMControllerActions.h
+++ b/Engine/Plugins/Runtime/RigVM/Source/RigVMDeveloper/Public/RigVMModel/RigVMControllerActions.h
@@ -929,6 +929,9 @@ public:
 	UPROPERTY()
 	FString OldCategory;
 
+	UPROPERTY()
+	int32 OldPinIndex;
+
 	UPROPERTY()
 	FString NewCategory;
 };
@@ -975,7 +978,7 @@ public:
 	FRigVMSetPinCategoryExpansionAction();
 	FRigVMSetPinCategoryExpansionAction(URigVMController* InController, const URigVMNode* InNode, const FString& InPinCategory);
 	virtual ~FRigVMSetPinCategoryExpansionAction() {};
-	virtual UScriptStruct* GetScriptStruct() const override { return FRigVMSetPinCategoryAction::StaticStruct(); }
+	virtual UScriptStruct* GetScriptStruct() const override { return FRigVMSetPinCategoryExpansionAction::StaticStruct(); }
 	virtual bool Undo() override;
 	virtual bool Redo() override;
 
diff --git a/Engine/Plugins/Runtime/RigVM/Source/RigVMEditor/Private/Widgets/SRigVMNodeLayoutWidget.cpp b/Engine/Plugins/Runtime/RigVM/Source/RigVMEditor/Private/Widgets/SRigVMNodeLayoutWidget.cpp
--- a/Engine/Plugins/Runtime/RigVM/Source/RigVMEditor/Private/Widgets/SRigVMNodeLayoutWidget.cpp
+++ b/Engine/Plugins/Runtime/RigVM/Source/RigVMEditor/Private/Widgets/SRigVMNodeLayoutWidget.cpp
@@ -330,6 +330,15 @@ void SRigVMNodeLayoutWidget::SRigVMNowLayoutRow::Construct(const FArguments& InA
 		.AutoWidth()
 		[
 			SNew(SButton)
+			.ToolTipText_Lambda([this]()
+			{
+				if(NodeLayoutRow->IsCategory())
+				{
+					return LOCTEXT("RemoveCategory", "Remove category");
+				}
+
+				return LOCTEXT("RemoveFromCategory", "Remove from category");
+			})
 			.ButtonStyle( FAppStyle::Get(), "SimpleButton" )
 			.Visibility_Lambda([this]()
 			{
