Fix UAF trait latent property handle access

diff --git a/Engine/Plugins/Experimental/UAF/UAFAnimGraph/Source/UAFAnimGraph/Private/TraitCore/TraitReader.cpp b/Engine/Plugins/Experimental/UAF/UAFAnimGraph/Source/UAFAnimGraph/Private/TraitCore/TraitReader.cpp
--- a/Engine/Plugins/Experimental/UAF/UAFAnimGraph/Source/UAFAnimGraph/Private/TraitCore/TraitReader.cpp
+++ b/Engine/Plugins/Experimental/UAF/UAFAnimGraph/Source/UAFAnimGraph/Private/TraitCore/TraitReader.cpp
@@ -130,11 +130,15 @@ namespace UE::UAF
 			{
 				const FTraitTemplate& TraitTemplate = TraitTemplates[TraitIndex];
 
-				uint32 NumLatentProperties = 0;
-				*this << NumLatentProperties;
+				// This is how many latent properties are defined on the trait
+				const uint32 NumTraitLatentProperties = TraitTemplate.GetNumLatentPropreties();
+
+				// This is how many latent properties were hooked to something on that node
+				uint32 NumSerializedLatentProperties = 0;
+				*this << NumSerializedLatentProperties;
 
 				// We allow traits to reserve more latent property handles than they might need
-				check(NumLatentProperties <= TraitTemplate.GetNumLatentPropreties());
+				check(NumSerializedLatentProperties <= NumTraitLatentProperties);
 
 				const uint32 BaseTraitIndex = TraitIndex - TraitTemplate.GetTraitIndex();
 				const FTraitTemplate& BaseTraitTemplate = TraitTemplates[BaseTraitIndex];
@@ -153,7 +157,13 @@ namespace UE::UAF
 					LatentHeader.bAllLatentPropertiesOnBecomeRelevant = true;
 				}
 
-				if (NumLatentProperties == 0)
+				// Default initialize our latent property handles
+				for (uint32 LatentPropertyIndex = 0; LatentPropertyIndex < NumTraitLatentProperties; ++LatentPropertyIndex)
+				{
+					LatentHandles[LatentPropertyIndex] = FLatentPropertyHandle();
+				}
+
+				if (NumSerializedLatentProperties == 0)
 				{
 					continue;	// Nothing to do
 				}
@@ -166,31 +176,23 @@ namespace UE::UAF
 				const FTraitRegistryHandle TraitHandle = TraitTemplate.GetRegistryHandle();
 				const FTrait* Trait = TraitRegistry.Find(TraitHandle);
 
-				if (Trait == nullptr)
-				{
-					// Default initialize our latent property handles if we don't have a valid trait
-					for (uint32 LatentPropertyIndex = 0; LatentPropertyIndex < NumLatentProperties; ++LatentPropertyIndex)
-					{
-						LatentHandles[LatentPropertyIndex] = FLatentPropertyHandle();
-					}
-				}
-
-				for (uint32 SerializedLatentPropertyIndex = 0; SerializedLatentPropertyIndex < NumLatentProperties; ++SerializedLatentPropertyIndex)
+				// Setup our latent property handles from the serialized data
+				for (uint32 SerializedLatentPropertyIndex = 0; SerializedLatentPropertyIndex < NumSerializedLatentProperties; ++SerializedLatentPropertyIndex)
 				{
 					FLatentPropertyMetadata Metadata;
 					*this << Metadata;
 
-					uint16 RigVMIndex = MAX_uint16;
-					uint32 CurrentLatentPropertyOffset = 0;
-					bool bCanFreeze = true;
-					bool bOnBecomeRelevant = true;
-
 					if (Trait != nullptr)
 					{
 						const FTraitLatentPropertyMemoryLayout PropertyMemoryLayout = Trait->GetLatentPropertyMemoryLayout(*SharedData, Metadata.Name, SerializedLatentPropertyIndex);
 						check(PropertyMemoryLayout.Size != 0);
 						check(PropertyMemoryLayout.Alignment != 0 && FMath::IsPowerOfTwo(PropertyMemoryLayout.Alignment));
-						check(PropertyMemoryLayout.LatentPropertyIndex != INDEX_NONE && (uint32)PropertyMemoryLayout.LatentPropertyIndex < NumLatentProperties);
+						check(PropertyMemoryLayout.LatentPropertyIndex != INDEX_NONE && (uint32)PropertyMemoryLayout.LatentPropertyIndex < NumTraitLatentProperties);
+
+						uint16 RigVMIndex = MAX_uint16;
+						uint32 CurrentLatentPropertyOffset = 0;
+						bool bCanFreeze = true;
+						bool bOnBecomeRelevant = true;
 
 						// If this property is valid, setup our binding for it
 						if (Metadata.RigVMIndex != MAX_uint16)
