Removed duplicate versions of FInstancedStructProvider now that the version in StructUtilsEditor is publicly available

#rnx

diff --git a/Engine/Plugins/Experimental/UAF/UAF/Source/UAFUncookedOnly/Private/Variables/UniversalObjectLocatorBindingType.cpp b/Engine/Plugins/Experimental/UAF/UAF/Source/UAFUncookedOnly/Private/Variables/UniversalObjectLocatorBindingType.cpp
--- a/Engine/Plugins/Experimental/UAF/UAF/Source/UAFUncookedOnly/Private/Variables/UniversalObjectLocatorBindingType.cpp
+++ b/Engine/Plugins/Experimental/UAF/UAF/Source/UAFUncookedOnly/Private/Variables/UniversalObjectLocatorBindingType.cpp
@@ -4,6 +4,7 @@
 
 #include "AnimNextModuleImpl.h"
 #include "ClassProxy.h"
+#include "InstancedStructDetails.h"
 #include "IStructureDataProvider.h"
 #include "IUniversalObjectLocatorEditorModule.h"
 #include "PropertyHandle.h"
@@ -40,140 +41,6 @@
 namespace UE::UAF::UncookedOnly
 {
 
-class FInstancedStructProvider : public IStructureDataProvider
-{
-public:
-	FInstancedStructProvider() = default;
-	
-	explicit FInstancedStructProvider(const TSharedPtr<IPropertyHandle>& InStructProperty)
-		: StructProperty(InStructProperty)
-	{
-	}
-	
-	virtual ~FInstancedStructProvider() override {}
-
-	void Reset()
-	{
-		StructProperty = nullptr;
-	}
-	
-	virtual bool IsValid() const override
-	{
-		bool bHasValidData = false;
-		EnumerateInstances([&bHasValidData](const UScriptStruct* ScriptStruct, uint8* Memory, UPackage* Package)
-		{
-			if (ScriptStruct && Memory)
-			{
-				bHasValidData = true;
-				return false; // Stop
-			}
-			return true; // Continue
-		});
-
-		return bHasValidData;
-	}
-	
-	virtual const UStruct* GetBaseStructure() const override
-	{
-		// Taken from UClass::FindCommonBase
-		auto FindCommonBaseStruct = [](const UScriptStruct* StructA, const UScriptStruct* StructB)
-		{
-			const UScriptStruct* CommonBaseStruct = StructA;
-			while (CommonBaseStruct && StructB && !StructB->IsChildOf(CommonBaseStruct))
-			{
-				CommonBaseStruct = Cast<UScriptStruct>(CommonBaseStruct->GetSuperStruct());
-			}
-			return CommonBaseStruct;
-		};
-
-		const UScriptStruct* CommonStruct = nullptr;
-		EnumerateInstances([&CommonStruct, &FindCommonBaseStruct](const UScriptStruct* ScriptStruct, uint8* Memory, UPackage* Package)
-		{
-			if (ScriptStruct)
-			{
-				CommonStruct = FindCommonBaseStruct(ScriptStruct, CommonStruct);
-			}
-			return true; // Continue
-		});
-
-		return CommonStruct;
-	}
-
-	virtual void GetInstances(TArray<TSharedPtr<FStructOnScope>>& OutInstances, const UStruct* ExpectedBaseStructure) const override
-	{
-		// The returned instances need to be compatible with base structure.
-		// This function returns empty instances in case they are not compatible, with the idea that we have as many instances as we have outer objects.
-		EnumerateInstances([&OutInstances, ExpectedBaseStructure](const UScriptStruct* ScriptStruct, uint8* Memory, UPackage* Package)
-		{
-			TSharedPtr<FStructOnScope> Result;
-			
-			if (ExpectedBaseStructure && ScriptStruct && ScriptStruct->IsChildOf(ExpectedBaseStructure))
-			{
-				Result = MakeShared<FStructOnScope>(ScriptStruct, Memory);
-				Result->SetPackage(Package);
-			}
-
-			OutInstances.Add(Result);
-
-			return true; // Continue
-		});
-	}
-
-	virtual bool IsPropertyIndirection() const override
-	{
-		return true;
-	}
-
-	virtual uint8* GetValueBaseAddress(uint8* ParentValueAddress, const UStruct* ExpectedBaseStructure) const override
-	{
-		if (!ParentValueAddress)
-		{
-			return nullptr;
-		}
-
-		FInstancedStruct& InstancedStruct = *reinterpret_cast<FInstancedStruct*>(ParentValueAddress);
-		if (ExpectedBaseStructure && InstancedStruct.GetScriptStruct() && InstancedStruct.GetScriptStruct()->IsChildOf(ExpectedBaseStructure))
-		{
-			return InstancedStruct.GetMutableMemory();
-		}
-		
-		return nullptr;
-	}
-	
-protected:
-
-	void EnumerateInstances(TFunctionRef<bool(const UScriptStruct* ScriptStruct, uint8* Memory, UPackage* Package)> InFunc) const
-	{
-		if (!StructProperty.IsValid())
-		{
-			return;
-		}
-		
-		TArray<UPackage*> Packages;
-		StructProperty->GetOuterPackages(Packages);
-
-		StructProperty->EnumerateRawData([&InFunc, &Packages](void* RawData, const int32 DataIndex, const int32 /*NumDatas*/)
-		{
-			const UScriptStruct* ScriptStruct = nullptr;
-			uint8* Memory = nullptr;
-			UPackage* Package = nullptr;
-			if (FInstancedStruct* InstancedStruct = static_cast<FInstancedStruct*>(RawData))
-			{
-				ScriptStruct = InstancedStruct->GetScriptStruct();
-				Memory = InstancedStruct->GetMutableMemory();
-				if (ensureMsgf(Packages.IsValidIndex(DataIndex), TEXT("Expecting packges and raw data to match.")))
-				{
-					Package = Packages[DataIndex];
-				}
-			}
-
-			return InFunc(ScriptStruct, Memory, Package);
-		});
-	}
-	
-	TSharedPtr<IPropertyHandle> StructProperty;
-};
-
 TSharedRef<SWidget> FUniversalObjectLocatorBindingType::CreateEditWidget(const TSharedRef<IPropertyHandle>& InPropertyHandle, const FAnimNextParamType& InType) const
 {
 	class SEditWidget : public SCompoundWidget
diff --git a/Engine/Plugins/Experimental/UAF/UAFAnimGraph/Source/UAFAnimGraphEditor/Private/Factory/AnimNextFactoryParamsDetails.cpp b/Engine/Plugins/Experimental/UAF/UAFAnimGraph/Source/UAFAnimGraphEditor/Private/Factory/AnimNextFactoryParamsDetails.cpp
--- a/Engine/Plugins/Experimental/UAF/UAFAnimGraph/Source/UAFAnimGraphEditor/Private/Factory/AnimNextFactoryParamsDetails.cpp
+++ b/Engine/Plugins/Experimental/UAF/UAFAnimGraph/Source/UAFAnimGraphEditor/Private/Factory/AnimNextFactoryParamsDetails.cpp
@@ -1,4 +1,4 @@
-﻿// Copyright Epic Games, Inc. All Rights Reserved.
+// Copyright Epic Games, Inc. All Rights Reserved.
 
 #include "AnimNextFactoryParamsDetails.h"
 
@@ -7,6 +7,7 @@
 #include "DetailWidgetRow.h"
 #include "IDetailChildrenBuilder.h"
 #include "IDetailPropertyRow.h"
+#include "InstancedStructDetails.h"
 #include "IStructureDataProvider.h"
 #include "PropertyHandle.h"
 #include "StructViewerFilter.h"
@@ -29,140 +30,6 @@
 namespace UE::UAF::Editor
 {
 
-class FInstancedStructProvider : public IStructureDataProvider
-{
-public:
-	FInstancedStructProvider() = default;
-	
-	explicit FInstancedStructProvider(const TSharedPtr<IPropertyHandle>& InStructProperty)
-		: StructProperty(InStructProperty)
-	{
-	}
-	
-	virtual ~FInstancedStructProvider() override {}
-
-	void Reset()
-	{
-		StructProperty = nullptr;
-	}
-	
-	virtual bool IsValid() const override
-	{
-		bool bHasValidData = false;
-		EnumerateInstances([&bHasValidData](const UScriptStruct* ScriptStruct, uint8* Memory, UPackage* Package)
-		{
-			if (ScriptStruct && Memory)
-			{
-				bHasValidData = true;
-				return false; // Stop
-			}
-			return true; // Continue
-		});
-
-		return bHasValidData;
-	}
-	
-	virtual const UStruct* GetBaseStructure() const override
-	{
-		// Taken from UClass::FindCommonBase
-		auto FindCommonBaseStruct = [](const UScriptStruct* StructA, const UScriptStruct* StructB)
-		{
-			const UScriptStruct* CommonBaseStruct = StructA;
-			while (CommonBaseStruct && StructB && !StructB->IsChildOf(CommonBaseStruct))
-			{
-				CommonBaseStruct = Cast<UScriptStruct>(CommonBaseStruct->GetSuperStruct());
-			}
-			return CommonBaseStruct;
-		};
-
-		const UScriptStruct* CommonStruct = nullptr;
-		EnumerateInstances([&CommonStruct, &FindCommonBaseStruct](const UScriptStruct* ScriptStruct, uint8* Memory, UPackage* Package)
-		{
-			if (ScriptStruct)
-			{
-				CommonStruct = FindCommonBaseStruct(ScriptStruct, CommonStruct);
-			}
-			return true; // Continue
-		});
-
-		return CommonStruct;
-	}
-
-	virtual void GetInstances(TArray<TSharedPtr<FStructOnScope>>& OutInstances, const UStruct* ExpectedBaseStructure) const override
-	{
-		// The returned instances need to be compatible with base structure.
-		// This function returns empty instances in case they are not compatible, with the idea that we have as many instances as we have outer objects.
-		EnumerateInstances([&OutInstances, ExpectedBaseStructure](const UScriptStruct* ScriptStruct, uint8* Memory, UPackage* Package)
-		{
-			TSharedPtr<FStructOnScope> Result;
-			
-			if (ExpectedBaseStructure && ScriptStruct && ScriptStruct->IsChildOf(ExpectedBaseStructure))
-			{
-				Result = MakeShared<FStructOnScope>(ScriptStruct, Memory);
-				Result->SetPackage(Package);
-			}
-
-			OutInstances.Add(Result);
-
-			return true; // Continue
-		});
-	}
-
-	virtual bool IsPropertyIndirection() const override
-	{
-		return true;
-	}
-
-	virtual uint8* GetValueBaseAddress(uint8* ParentValueAddress, const UStruct* ExpectedBaseStructure) const override
-	{
-		if (!ParentValueAddress)
-		{
-			return nullptr;
-		}
-
-		FInstancedStruct& InstancedStruct = *reinterpret_cast<FInstancedStruct*>(ParentValueAddress);
-		if (ExpectedBaseStructure && InstancedStruct.GetScriptStruct() && InstancedStruct.GetScriptStruct()->IsChildOf(ExpectedBaseStructure))
-		{
-			return InstancedStruct.GetMutableMemory();
-		}
-		
-		return nullptr;
-	}
-	
-protected:
-
-	void EnumerateInstances(TFunctionRef<bool(const UScriptStruct* ScriptStruct, uint8* Memory, UPackage* Package)> InFunc) const
-	{
-		if (!StructProperty.IsValid() || !StructProperty->IsValidHandle())
-		{
-			return;
-		}
-		
-		TArray<UPackage*> Packages;
-		StructProperty->GetOuterPackages(Packages);
-
-		StructProperty->EnumerateRawData([&InFunc, &Packages](void* RawData, const int32 DataIndex, const int32 /*NumDatas*/)
-		{
-			const UScriptStruct* ScriptStruct = nullptr;
-			uint8* Memory = nullptr;
-			UPackage* Package = nullptr;
-			if (FInstancedStruct* InstancedStruct = static_cast<FInstancedStruct*>(RawData))
-			{
-				ScriptStruct = InstancedStruct->GetScriptStruct();
-				Memory = InstancedStruct->GetMutableMemory();
-				if (ensureMsgf(Packages.IsValidIndex(DataIndex), TEXT("Expecting packges and raw data to match.")))
-				{
-					Package = Packages[DataIndex];
-				}
-			}
-
-			return InFunc(ScriptStruct, Memory, Package);
-		});
-	}
-	
-	TSharedPtr<IPropertyHandle> StructProperty;
-};
-
 void FAnimNextFactoryParamsDetails::CustomizeHeader(TSharedRef<IPropertyHandle> InPropertyHandle, FDetailWidgetRow& InHeaderRow, IPropertyTypeCustomizationUtils& InCustomizationUtils)
 {
 }
diff --git a/Engine/Plugins/Experimental/UAF/UAFAnimGraph/Source/UAFAnimGraphEditor/UAFAnimGraphEditor.Build.cs b/Engine/Plugins/Experimental/UAF/UAFAnimGraph/Source/UAFAnimGraphEditor/UAFAnimGraphEditor.Build.cs
--- a/Engine/Plugins/Experimental/UAF/UAFAnimGraph/Source/UAFAnimGraphEditor/UAFAnimGraphEditor.Build.cs
+++ b/Engine/Plugins/Experimental/UAF/UAFAnimGraph/Source/UAFAnimGraphEditor/UAFAnimGraphEditor.Build.cs
@@ -28,6 +28,7 @@ namespace UnrealBuildTool.Rules
 					"UAFAnimGraph",
 					"UAFAnimGraphUncookedOnly",
 					"Settings",
+					"StructUtilsEditor",
 					"WorkspaceEditor", 
 					"EditorWidgets", 
 					"MessageLog", 
diff --git a/Engine/Source/Editor/StructUtilsEditor/Private/InstancedStructDetails.cpp b/Engine/Source/Editor/StructUtilsEditor/Private/InstancedStructDetails.cpp
--- a/Engine/Source/Editor/StructUtilsEditor/Private/InstancedStructDetails.cpp
+++ b/Engine/Source/Editor/StructUtilsEditor/Private/InstancedStructDetails.cpp
@@ -114,7 +114,7 @@ uint8* FInstancedStructProvider::GetValueBaseAddress(uint8* ParentValueAddress,
 
 void FInstancedStructProvider::EnumerateInstances(TFunctionRef<bool(const UScriptStruct* ScriptStruct, uint8* Memory, UPackage* Package)> InFunc) const
 {
-	if (!StructProperty.IsValid())
+	if (!StructProperty.IsValid() || !StructProperty->IsValidHandle())
 	{
 		return;
 	}
