Fix rewind debugger crash if sequence is missing from sampling task

diff --git a/Engine/Plugins/Experimental/UAF/UAFAnimGraph/Source/UAFAnimGraphEditor/Private/RewindDebugger/SequenceInfoTrack.cpp b/Engine/Plugins/Experimental/UAF/UAFAnimGraph/Source/UAFAnimGraphEditor/Private/RewindDebugger/SequenceInfoTrack.cpp
--- a/Engine/Plugins/Experimental/UAF/UAFAnimGraph/Source/UAFAnimGraphEditor/Private/RewindDebugger/SequenceInfoTrack.cpp
+++ b/Engine/Plugins/Experimental/UAF/UAFAnimGraph/Source/UAFAnimGraphEditor/Private/RewindDebugger/SequenceInfoTrack.cpp
@@ -257,14 +257,18 @@ void FSequenceInfoTrack::RefreshSequenceInfoFromEvaluationProgram(TArray<FAnimNe
 			const FAnimNextAnimSequenceKeyframeTask* SequenceTask = Task.GetPtr<FAnimNextAnimSequenceKeyframeTask>();
 			FAnimNextSequenceTraceInfo& TraceInfo = OutSequenceInfo.Emplace_GetRef();
 			TraceInfo.AnimSequence = SequenceTask->AnimSequence;
-			TraceInfo.DurationSeconds = SequenceTask->AnimSequence.IsValid() ? TraceInfo.AnimSequence->GetPlayLength() : 0.0f;
 			TraceInfo.CurrentTimeSeconds = SequenceTask->SampleTime;
 
-			// Fill out sync markers
-			TraceInfo.SyncMarkers.Reserve(SequenceTask->AnimSequence->AuthoredSyncMarkers.Num());
-			for (auto Marker : SequenceTask->AnimSequence->AuthoredSyncMarkers)
+			if (const UAnimSequence* AnimSequence = SequenceTask->AnimSequence.Get())
 			{
-				TraceInfo.SyncMarkers.Emplace(Marker.Time, Marker.MarkerName);
+				TraceInfo.DurationSeconds = AnimSequence->GetPlayLength();
+
+				// Fill out sync markers
+				TraceInfo.SyncMarkers.Reserve(AnimSequence->AuthoredSyncMarkers.Num());
+				for (auto Marker : AnimSequence->AuthoredSyncMarkers)
+				{
+					TraceInfo.SyncMarkers.Emplace(Marker.Time, Marker.MarkerName);
+				}
 			}
 		}
 	}
