UAF: Fix for blend stack animations not looping

diff --git a/Engine/Plugins/Experimental/UAF/UAFAnimGraph/Source/UAFAnimGraph/Private/AnimNextAnimGraphModule.cpp b/Engine/Plugins/Experimental/UAF/UAFAnimGraph/Source/UAFAnimGraph/Private/AnimNextAnimGraphModule.cpp
--- a/Engine/Plugins/Experimental/UAF/UAFAnimGraph/Source/UAFAnimGraph/Private/AnimNextAnimGraphModule.cpp
+++ b/Engine/Plugins/Experimental/UAF/UAFAnimGraph/Source/UAFAnimGraph/Private/AnimNextAnimGraphModule.cpp
@@ -111,6 +111,10 @@ void FAnimNextAnimGraphModule::StartupModule()
 				ensure(InContext.AccessStruct<FSequencePlayerData>([&InContext](FSequencePlayerData& InStruct)
 				{
 					InStruct.AnimSequence = CastChecked<UAnimSequence>(InContext.GetObject());
+					if (InStruct.AnimSequence)
+					{
+						InStruct.bLoop = InStruct.AnimSequence->bLoop;
+					}
 				}));
 			});
 	}
@@ -128,6 +132,10 @@ void FAnimNextAnimGraphModule::StartupModule()
 				ensure(InContext.AccessStruct<FBlendSpacePlayerData>([&InContext](FBlendSpacePlayerData& InStruct)
 				{
 					InStruct.BlendSpace = CastChecked<UBlendSpace>(InContext.GetObject());
+					if (InStruct.BlendSpace)
+					{
+						InStruct.bLoop = InStruct.BlendSpace->bLoop;
+					}
 				}));
 			});
 	}
