UAF: Fix expected trait name in automated test

diff --git a/Engine/Plugins/Experimental/UAF/UAFAnimGraph/Source/UAFAnimGraphTestSuite/Private/AnimNextAnimGraphTraitGraphTest.cpp b/Engine/Plugins/Experimental/UAF/UAFAnimGraph/Source/UAFAnimGraphTestSuite/Private/AnimNextAnimGraphTraitGraphTest.cpp
--- a/Engine/Plugins/Experimental/UAF/UAFAnimGraph/Source/UAFAnimGraphTestSuite/Private/AnimNextAnimGraphTraitGraphTest.cpp
+++ b/Engine/Plugins/Experimental/UAF/UAFAnimGraph/Source/UAFAnimGraphTestSuite/Private/AnimNextAnimGraphTraitGraphTest.cpp
@@ -257,9 +257,10 @@ bool FAnimationAnimNextEditorTest_GraphTraitOperations::RunTest(const FString& I
 			UE_RETURN_ON_ERROR(Trait != nullptr, "FAnimationAnimNextEditorTest_GraphTraitOperations -> Failed to find test trait");
 
 			const FName TraitTypeName = *Trait->GetTraitName();
-
+			const FName TraitSharedDataStructName = Trait->GetTraitSharedDataStruct()->GetFName();
+			
 			TraitInstanceName = Controller->AddTraitByName(TraitStackNode->GetFName(), TraitTypeName, INDEX_NONE);
-			UE_RETURN_ON_ERROR(TraitInstanceName == TraitTypeName, TEXT("FAnimationAnimNextEditorTest_GraphTraitOperations -> Unexpected Trait name"));
+			UE_RETURN_ON_ERROR(TraitInstanceName == TraitSharedDataStructName, TEXT("FAnimationAnimNextEditorTest_GraphTraitOperations -> Unexpected Trait name"));
 
 			URigVMPin* TraitPin = TraitStackNode->FindPin(TraitInstanceName.ToString());
 			UE_RETURN_ON_ERROR(TraitPin != nullptr, TEXT("FAnimationAnimNextEditorTest_GraphTraitOperations -> Failed to find Trait pin"));
@@ -382,9 +383,10 @@ bool FAnimationAnimNextEditorTest_GraphTraitOperations::RunTest(const FString& I
 			UE_RETURN_ON_ERROR(BasicTrait != nullptr, "FAnimationAnimNextEditorTest_GraphTraitOperations -> Failed to find test basic trait");
 
 			const FName BasicTraitTypeName = *BasicTrait->GetTraitName();
+			const FName BasicTraitSharedDataStructName = BasicTrait->GetTraitSharedDataStruct()->GetFName();
 
 			TraitInstanceName = Controller->SwapTraitByName(TraitStackNode->GetFName(), TraitInstanceName, 1, BasicTraitTypeName);
-			UE_RETURN_ON_ERROR(TraitInstanceName == BasicTraitTypeName, TEXT("FAnimationAnimNextEditorTest_GraphTraitOperations -> Unexpected Trait name"));
+			UE_RETURN_ON_ERROR(TraitInstanceName == BasicTraitSharedDataStructName, TEXT("FAnimationAnimNextEditorTest_GraphTraitOperations -> Unexpected Trait name"));
 
 			URigVMPin* TraitPin = TraitStackNode->FindPin(TraitInstanceName.ToString());
 			UE_RETURN_ON_ERROR(TraitPin != nullptr, TEXT("FAnimationAnimNextEditorTest_GraphTraitOperations -> Failed to find FTestBasicTrait pin"));
