Properly handle missing latent properties in UAF trait declaration
Add missing latent property

diff --git a/Engine/Plugins/Experimental/UAF/UAFAnimGraph/Source/UAFAnimGraph/Private/Traits/ModifyCurveTrait.h b/Engine/Plugins/Experimental/UAF/UAFAnimGraph/Source/UAFAnimGraph/Private/Traits/ModifyCurveTrait.h
--- a/Engine/Plugins/Experimental/UAF/UAFAnimGraph/Source/UAFAnimGraph/Private/Traits/ModifyCurveTrait.h
+++ b/Engine/Plugins/Experimental/UAF/UAFAnimGraph/Source/UAFAnimGraph/Private/Traits/ModifyCurveTrait.h
@@ -55,6 +55,7 @@ struct FAnimNextModifyCurveTraitSharedData : public FAnimNextTraitSharedData
 	float Alpha = 1.0f;
 
 #define TRAIT_LATENT_PROPERTIES_ENUMERATOR(GeneratorMacro) \
+	GeneratorMacro(ApplyMode) \
 	GeneratorMacro(Alpha) \
 
 	GENERATE_TRAIT_LATENT_PROPERTIES(FAnimNextModifyCurveTraitSharedData, TRAIT_LATENT_PROPERTIES_ENUMERATOR)
diff --git a/Engine/Plugins/Experimental/UAF/UAFAnimGraph/Source/UAFAnimGraph/Public/TraitCore/Trait.h b/Engine/Plugins/Experimental/UAF/UAFAnimGraph/Source/UAFAnimGraph/Public/TraitCore/Trait.h
--- a/Engine/Plugins/Experimental/UAF/UAFAnimGraph/Source/UAFAnimGraph/Public/TraitCore/Trait.h
+++ b/Engine/Plugins/Experimental/UAF/UAFAnimGraph/Source/UAFAnimGraph/Public/TraitCore/Trait.h
@@ -117,8 +117,18 @@ struct FRigVMPinInfoArray;
 	{ \
 		const UScriptStruct* SharedDataStruct = FSharedData::StaticStruct(); \
 		const FProperty* Property = SharedDataStruct->FindPropertyByName(PropertyName); \
-		ensure(Property != nullptr); \
-		return Property != nullptr ? (FSharedData::GetLatentPropertyIndex(Property->GetOffset_ForInternal()) - 1) : INDEX_NONE; \
+		if (!ensure(Property != nullptr)) \
+		{ \
+			/* Unknown property */ \
+			return INDEX_NONE; \
+		} \
+		const int32 LatentPropertyIndex = FSharedData::GetLatentPropertyIndex(Property->GetOffset_ForInternal()); \
+		if (LatentPropertyIndex <= 0) \
+		{ \
+			/* Property not part of latent enumerator macro */ \
+			return INDEX_NONE; \
+		} \
+		return LatentPropertyIndex - 1; \
 	}
 
 // Helper that handles the GetTraitInterface() details for each interface specified by the generator macro
