Ensure Test/Shipping configuration works with UAF proceduraly generated graphs

diff --git a/Engine/Plugins/Experimental/UAF/UAFAnimGraph/Source/UAFAnimGraph/Private/Factory/AnimGraphBuilderContext.cpp b/Engine/Plugins/Experimental/UAF/UAFAnimGraph/Source/UAFAnimGraph/Private/Factory/AnimGraphBuilderContext.cpp
--- a/Engine/Plugins/Experimental/UAF/UAFAnimGraph/Source/UAFAnimGraph/Private/Factory/AnimGraphBuilderContext.cpp
+++ b/Engine/Plugins/Experimental/UAF/UAFAnimGraph/Source/UAFAnimGraph/Private/Factory/AnimGraphBuilderContext.cpp
@@ -1,4 +1,4 @@
-﻿// Copyright Epic Games, Inc. All Rights Reserved.
+// Copyright Epic Games, Inc. All Rights Reserved.
 
 #include "Factory/AnimGraphBuilderContext.h"
 #include "Graph/AnimNextAnimationGraph.h"
@@ -28,7 +28,7 @@ const UAnimNextAnimationGraph* FAnimGraphBuilderContext::Build()
 	AnimationGraph->ReferencedVariableStructs = VariableStructs;
 	AnimationGraph->GraphReferencedObjects = TraitWriter.GetGraphReferencedObjects();
 	AnimationGraph->GraphReferencedSoftObjects = TraitWriter.GetGraphReferencedSoftObjects();
-	check(AnimationGraph->LoadFromArchiveBuffer(TraitWriter.GetGraphSharedData()));
+	ensure(AnimationGraph->LoadFromArchiveBuffer(TraitWriter.GetGraphSharedData()));
 
 	// Clear async flag to make the object visible to GC
 	(void)AnimationGraph->AtomicallyClearInternalFlags(EInternalObjectFlags::Async);
