Fix various trait stack issues
* Fix being able to rearrange a trait onto itself
* Prevent rearranging via OnItem as behaviour is unclear
* Fix reordering traits not working properly

diff --git a/Engine/Plugins/Experimental/UAF/UAFAnimGraph/Source/UAFAnimGraphEditor/Private/TraitStackEditor.cpp b/Engine/Plugins/Experimental/UAF/UAFAnimGraph/Source/UAFAnimGraphEditor/Private/TraitStackEditor.cpp
--- a/Engine/Plugins/Experimental/UAF/UAFAnimGraph/Source/UAFAnimGraphEditor/Private/TraitStackEditor.cpp
+++ b/Engine/Plugins/Experimental/UAF/UAFAnimGraph/Source/UAFAnimGraphEditor/Private/TraitStackEditor.cpp
@@ -108,9 +108,8 @@ namespace Private
 			switch (InItemDropZone)
 			{
 			case EItemDropZone::AboveItem: return &TreeViewStyle.DropIndicator_Above;
-			default:
-			case EItemDropZone::OntoItem: return &TreeViewStyle.DropIndicator_Onto;
 			case EItemDropZone::BelowItem: return &TreeViewStyle.DropIndicator_Below;
+			default: return nullptr;
 			};
 		}
 		
@@ -138,8 +137,8 @@ namespace Private
 			return LayerId;
 		}
 
-		/** @return the zone (above, onto, below) based on where the user is hovering over within the row */
-		static EItemDropZone ZoneFromPointerPosition(UE::Slate::FDeprecateVector2DParameter LocalPointerPos, UE::Slate::FDeprecateVector2DParameter LocalSize)
+		/** @return the zone (above, below, none) based on where the user is hovering over within the row */
+		static TOptional<EItemDropZone> ZoneFromPointerPosition(UE::Slate::FDeprecateVector2DParameter LocalPointerPos, UE::Slate::FDeprecateVector2DParameter LocalSize)
 		{
 			const float ZoneBoundarySu = FMath::Clamp(LocalSize.Y * 0.25f, 3.0f, 10.0f);
 			if (LocalPointerPos.Y < ZoneBoundarySu)
@@ -152,7 +151,7 @@ namespace Private
 			}
 			else
 			{
-				return EItemDropZone::OntoItem;
+				return TOptional<EItemDropZone>();
 			}
 		}
 
@@ -180,6 +179,12 @@ namespace Private
 			{
 				return FReply::Unhandled();
 			}
+
+			if (DragOp->WeakPropertyHandle == WeakPropertyHandle)
+			{
+				ItemDropZone.Reset();
+				return FReply::Handled();
+			}
 			
 			const FVector2f LocalPointerPos = MyGeometry.AbsoluteToLocal(DragDropEvent.GetScreenSpacePosition());
 			ItemDropZone = ZoneFromPointerPosition(LocalPointerPos, MyGeometry.GetLocalSize());
@@ -258,16 +263,10 @@ namespace Private
 							if (VMTrait->GetTraitSharedDataStruct() == TargetStructData->GetStruct())
 							{
 								TargetPinIndex = Pin->GetPinIndex();
-								switch (ItemDropZone.GetValue())
+
+								if (ItemDropZone.IsSet() && ItemDropZone.GetValue() == EItemDropZone::AboveItem)
 								{
-								case EItemDropZone::AboveItem:
 									TargetPinIndex = FMath::Max(0, TargetPinIndex - 1);
-									break;
-								case EItemDropZone::OntoItem:
-									break;
-								case EItemDropZone::BelowItem:
-									TargetPinIndex = FMath::Min(TraitPins.Num(), TargetPinIndex + 1);
-									break;
 								}
 
 								if (SourcePinIndex != TargetPinIndex)
diff --git a/Engine/Plugins/Experimental/UAF/UAFAnimGraph/Source/UAFAnimGraphUncookedOnly/Private/AnimNextController.cpp b/Engine/Plugins/Experimental/UAF/UAFAnimGraph/Source/UAFAnimGraphUncookedOnly/Private/AnimNextController.cpp
--- a/Engine/Plugins/Experimental/UAF/UAFAnimGraph/Source/UAFAnimGraphUncookedOnly/Private/AnimNextController.cpp
+++ b/Engine/Plugins/Experimental/UAF/UAFAnimGraph/Source/UAFAnimGraphUncookedOnly/Private/AnimNextController.cpp
@@ -266,8 +266,16 @@ bool UAnimNextController::SetTraitPinIndex(FName InNodeName, FName InTraitInstan
 			{
 				// Avoid multiple VM recompilations, for each operation
 				FRigVMControllerCompileBracketScope CompileScope(this);
+				const int32 TraitPinInitialIndex = TraitPin->GetPinIndex();
+				
 				if (RemoveTraitByName(InNodeName, InTraitInstanceName, bSetupUndoRedo, bPrintPythonCommand))
 				{
+					// Adjust indices due to pin removal
+					if (TraitPinInitialIndex > InNewPinIndex)
+					{
+						++InNewPinIndex;
+					}
+					
 					// TOOO zzz : Why TraitName is a string ?
 					if (AddTraitByName(InNodeName, FName(AnimNextTrait->GetTraitName()), InNewPinIndex, TraitDefaultValue, bSetupUndoRedo, bPrintPythonCommand) != NAME_None)
 					{
