Invert logic to improve clarity

diff --git a/Engine/Plugins/Experimental/UAF/UAFAnimGraph/Source/UAFAnimGraphEditor/Private/TraitStackEditor.cpp b/Engine/Plugins/Experimental/UAF/UAFAnimGraph/Source/UAFAnimGraphEditor/Private/TraitStackEditor.cpp
--- a/Engine/Plugins/Experimental/UAF/UAFAnimGraph/Source/UAFAnimGraphEditor/Private/TraitStackEditor.cpp
+++ b/Engine/Plugins/Experimental/UAF/UAFAnimGraph/Source/UAFAnimGraphEditor/Private/TraitStackEditor.cpp
@@ -264,9 +264,9 @@ namespace Private
 							{
 								TargetPinIndex = Pin->GetPinIndex();
 
-								if (ItemDropZone.IsSet() && ItemDropZone.GetValue() == EItemDropZone::AboveItem)
+								if (ItemDropZone.IsSet() && ItemDropZone.GetValue() == EItemDropZone::BelowItem)
 								{
-									TargetPinIndex = FMath::Max(0, TargetPinIndex - 1);
+									TargetPinIndex = FMath::Max(0, TargetPinIndex + 1);
 								}
 
 								if (SourcePinIndex != TargetPinIndex)
diff --git a/Engine/Plugins/Experimental/UAF/UAFAnimGraph/Source/UAFAnimGraphUncookedOnly/Private/AnimNextController.cpp b/Engine/Plugins/Experimental/UAF/UAFAnimGraph/Source/UAFAnimGraphUncookedOnly/Private/AnimNextController.cpp
--- a/Engine/Plugins/Experimental/UAF/UAFAnimGraph/Source/UAFAnimGraphUncookedOnly/Private/AnimNextController.cpp
+++ b/Engine/Plugins/Experimental/UAF/UAFAnimGraph/Source/UAFAnimGraphUncookedOnly/Private/AnimNextController.cpp
@@ -271,9 +271,9 @@ bool UAnimNextController::SetTraitPinIndex(FName InNodeName, FName InTraitInstan
 				if (RemoveTraitByName(InNodeName, InTraitInstanceName, bSetupUndoRedo, bPrintPythonCommand))
 				{
 					// Adjust indices due to pin removal
-					if (TraitPinInitialIndex > InNewPinIndex)
+					if (TraitPinInitialIndex < InNewPinIndex)
 					{
-						++InNewPinIndex;
+						--InNewPinIndex;
 					}
 					
 					// TOOO zzz : Why TraitName is a string ?
