Avoid copying pthread lock as it leads to undefined behavior

diff --git a/Engine/Plugins/Experimental/UAF/UAFAnimGraph/Source/UAFAnimGraph/Internal/TraitCore/NodeTemplateRegistry.h b/Engine/Plugins/Experimental/UAF/UAFAnimGraph/Source/UAFAnimGraph/Internal/TraitCore/NodeTemplateRegistry.h
--- a/Engine/Plugins/Experimental/UAF/UAFAnimGraph/Source/UAFAnimGraph/Internal/TraitCore/NodeTemplateRegistry.h
+++ b/Engine/Plugins/Experimental/UAF/UAFAnimGraph/Source/UAFAnimGraph/Internal/TraitCore/NodeTemplateRegistry.h
@@ -26,6 +26,8 @@ namespace UE::UAF
 	 */
 	struct FNodeTemplateRegistry final
 	{
+		UE_NONCOPYABLE(FNodeTemplateRegistry);	// We disallow copy/move semantics
+
 		// Access the global registry
 		static UE_API FNodeTemplateRegistry& Get();
 
@@ -46,10 +48,6 @@ namespace UE::UAF
 
 	private:
 		FNodeTemplateRegistry() = default;
-		FNodeTemplateRegistry(const FNodeTemplateRegistry&) = delete;
-		FNodeTemplateRegistry(FNodeTemplateRegistry&&) = delete;
-		FNodeTemplateRegistry& operator=(const FNodeTemplateRegistry&) = delete;
-		FNodeTemplateRegistry& operator=(FNodeTemplateRegistry&&) = delete;
 
 		// Finds and returns a node template based on its handle or nullptr if the handle is invalid
 		UE_API FNodeTemplate* FindMutable(FNodeTemplateRegistryHandle TemplateHandle);
@@ -58,6 +56,12 @@ namespace UE::UAF
 		static UE_API void Init();
 		static UE_API void Destroy();
 
+#if WITH_DEV_AUTOMATION_TESTS
+		// Swaps the global node template registry with the specified instance
+		// Swap it back for the original once tests are done
+		static UE_API FNodeTemplateRegistry* Swap(FNodeTemplateRegistry* Other);
+#endif
+
 		// We only ever append node templates to this contiguous buffers
 		// This is an optimization, we share these and by being contiguous
 		// we improve cache locality and cache line density
diff --git a/Engine/Plugins/Experimental/UAF/UAFAnimGraph/Source/UAFAnimGraph/Private/TraitCore/NodeTemplateRegistry.cpp b/Engine/Plugins/Experimental/UAF/UAFAnimGraph/Source/UAFAnimGraph/Private/TraitCore/NodeTemplateRegistry.cpp
--- a/Engine/Plugins/Experimental/UAF/UAFAnimGraph/Source/UAFAnimGraph/Private/TraitCore/NodeTemplateRegistry.cpp
+++ b/Engine/Plugins/Experimental/UAF/UAFAnimGraph/Source/UAFAnimGraph/Private/TraitCore/NodeTemplateRegistry.cpp
@@ -36,6 +36,15 @@ namespace UE::UAF
 		}
 	}
 
+#if WITH_DEV_AUTOMATION_TESTS
+	FNodeTemplateRegistry* FNodeTemplateRegistry::Swap(FNodeTemplateRegistry* Other)
+	{
+		FNodeTemplateRegistry* Original = Private::GNodeTemplateRegistry;
+		Private::GNodeTemplateRegistry = Other;
+		return Original;
+	}
+#endif
+
 	FNodeTemplateRegistryHandle FNodeTemplateRegistry::Find(uint32 NodeTemplateUID) const
 	{
 		FReadScopeLock Lock(TemplateLock);
diff --git a/Engine/Plugins/Experimental/UAF/UAFAnimGraph/Source/UAFAnimGraphTestSuite/Private/AnimNextRuntimeTest.cpp b/Engine/Plugins/Experimental/UAF/UAFAnimGraph/Source/UAFAnimGraphTestSuite/Private/AnimNextRuntimeTest.cpp
--- a/Engine/Plugins/Experimental/UAF/UAFAnimGraph/Source/UAFAnimGraphTestSuite/Private/AnimNextRuntimeTest.cpp
+++ b/Engine/Plugins/Experimental/UAF/UAFAnimGraph/Source/UAFAnimGraphTestSuite/Private/AnimNextRuntimeTest.cpp
@@ -16,16 +16,13 @@ namespace UE::UAF
 {
 	FScopedClearNodeTemplateRegistry::FScopedClearNodeTemplateRegistry()
 	{
-		FNodeTemplateRegistry& Registry = FNodeTemplateRegistry::Get();
-		// Can't use Swap because the move constructor is private and Swap thinks it's non-swappable
-		FMemory::Memswap(&Registry, &TmpRegistry, sizeof(FNodeTemplateRegistry));
+		OriginalRegistry = FNodeTemplateRegistry::Swap(&TmpRegistry);
 	}
 
 	FScopedClearNodeTemplateRegistry::~FScopedClearNodeTemplateRegistry()
 	{
-		FNodeTemplateRegistry& Registry = FNodeTemplateRegistry::Get();
-		// Can't use Swap because the move constructor is private and Swap thinks it's non-swappable
-		FMemory::Memswap(&Registry, &TmpRegistry, sizeof(FNodeTemplateRegistry));
+		FNodeTemplateRegistry* PrevRegistry = FNodeTemplateRegistry::Swap(OriginalRegistry);
+		checkf(PrevRegistry == &TmpRegistry, TEXT("Unexpected node template registry instance found"));
 	}
 
 	bool FTestUtils::LoadFromArchiveBuffer(UAnimNextAnimationGraph& AnimationGraph, TArray<FNodeHandle>& NodeHandles, const TArray<uint8>& SharedDataArchiveBuffer)
diff --git a/Engine/Plugins/Experimental/UAF/UAFAnimGraph/Source/UAFAnimGraphTestSuite/Private/AnimNextRuntimeTest.h b/Engine/Plugins/Experimental/UAF/UAFAnimGraph/Source/UAFAnimGraphTestSuite/Private/AnimNextRuntimeTest.h
--- a/Engine/Plugins/Experimental/UAF/UAFAnimGraph/Source/UAFAnimGraphTestSuite/Private/AnimNextRuntimeTest.h
+++ b/Engine/Plugins/Experimental/UAF/UAFAnimGraph/Source/UAFAnimGraphTestSuite/Private/AnimNextRuntimeTest.h
@@ -77,7 +77,11 @@ namespace UE::UAF
 		UAFANIMGRAPHTESTSUITE_API FScopedClearNodeTemplateRegistry();
 		UAFANIMGRAPHTESTSUITE_API ~FScopedClearNodeTemplateRegistry();
 
+		// The new temporary node registry instance we'll set as the global one
 		FNodeTemplateRegistry TmpRegistry;
+
+		// The old node registry instance we'll restore once we go out of scope
+		FNodeTemplateRegistry* OriginalRegistry = nullptr;
 	};
 
 	struct FTestUtils final
diff --git a/Engine/Plugins/Experimental/UAF/UAFAnimGraph/Source/UAFAnimGraphTestSuite/Private/AnimNextTraitBaseTest.cpp b/Engine/Plugins/Experimental/UAF/UAFAnimGraph/Source/UAFAnimGraphTestSuite/Private/AnimNextTraitBaseTest.cpp
--- a/Engine/Plugins/Experimental/UAF/UAFAnimGraph/Source/UAFAnimGraphTestSuite/Private/AnimNextTraitBaseTest.cpp
+++ b/Engine/Plugins/Experimental/UAF/UAFAnimGraph/Source/UAFAnimGraphTestSuite/Private/AnimNextTraitBaseTest.cpp
@@ -1155,7 +1155,6 @@ bool FAnimationAnimNextRuntimeTest_TraitSerialization::RunTest(const FString& In
 
 	{
 		FScopedClearNodeTemplateRegistry ScopedClearNodeTemplateRegistry;
-		FNodeTemplateRegistry& Registry = FNodeTemplateRegistry::Get();
 
 		AUTO_REGISTER_ANIM_TRAIT(FTraitSerialization_Base)
 		AUTO_REGISTER_ANIM_TRAIT(FTraitSerialization_Add)
@@ -1174,7 +1173,7 @@ bool FAnimationAnimNextRuntimeTest_TraitSerialization::RunTest(const FString& In
 		TArray<uint8> NodeTemplateBuffer0;
 		const FNodeTemplate* NodeTemplate0 = FNodeTemplateBuilder::BuildNodeTemplate(NodeTemplateTraitList, NodeTemplateBuffer0);
 
-		FNodeTemplateRegistryHandle TemplateHandle0 = Registry.FindOrAdd(NodeTemplate0);
+		FNodeTemplateRegistryHandle TemplateHandle0 = FNodeTemplateRegistry::Get().FindOrAdd(NodeTemplate0);
 		AddErrorIfFalse(TemplateHandle0.IsValid(), "FAnimationAnimNextRuntimeTest_TraitSerialization -> Registry should contain our template");
 
 		FTraitSerialization_Base::FSharedData TraitBaseRef0;
@@ -1365,7 +1364,7 @@ bool FAnimationAnimNextRuntimeTest_TraitSerialization::RunTest(const FString& In
 		{
 			FScopedClearNodeTemplateRegistry ScopedClearNodeTemplateRegistryForLoad;
 
-			AddErrorIfFalse(Registry.GetNum() == 0, "FAnimationAnimNextRuntimeTest_TraitSerialization -> Registry should contain 0 templates");
+			AddErrorIfFalse(FNodeTemplateRegistry::Get().GetNum() == 0, "FAnimationAnimNextRuntimeTest_TraitSerialization -> Registry should contain 0 templates");
 
 			// Read our graph
 			FTestUtils::LoadFromArchiveBuffer(*AnimationGraph, NodeHandles, GraphSharedDataArchiveBuffer);
@@ -1511,7 +1510,6 @@ bool FAnimationAnimNextRuntimeTest_SimpleGraphBuilder::RunTest(const FString& In
 
 	{
 		FScopedClearNodeTemplateRegistry ScopedClearNodeTemplateRegistry;
-		FNodeTemplateRegistry& Registry = FNodeTemplateRegistry::Get();
 
 		AUTO_REGISTER_ANIM_TRAIT(FTraitSerialization_Base)
 		AUTO_REGISTER_ANIM_TRAIT(FTraitSerialization_Add)
@@ -1574,6 +1572,7 @@ bool FAnimationAnimNextRuntimeTest_SimpleGraphBuilder::RunTest(const FString& In
 			// Clear out the node template registry to test registration on load
 			{
 				FScopedClearNodeTemplateRegistry ScopedClearNodeTemplateRegistryForLoad;
+				FNodeTemplateRegistry& Registry = FNodeTemplateRegistry::Get();
 
 				AddErrorIfFalse(Registry.GetNum() == 0, "FAnimationAnimNextRuntimeTest_SimpleGraphBuilder -> Registry should contain 0 templates");
 
