Avoid crash when force deleting a UAF animation graph

diff --git a/Engine/Plugins/Experimental/UAF/UAF/Source/UAF/Private/Module/AnimNextModuleInstance.cpp b/Engine/Plugins/Experimental/UAF/UAF/Source/UAF/Private/Module/AnimNextModuleInstance.cpp
--- a/Engine/Plugins/Experimental/UAF/UAF/Source/UAF/Private/Module/AnimNextModuleInstance.cpp
+++ b/Engine/Plugins/Experimental/UAF/UAF/Source/UAF/Private/Module/AnimNextModuleInstance.cpp
@@ -503,7 +503,7 @@ void FAnimNextModuleInstance::Trace()
 
 const UAnimNextModule* FAnimNextModuleInstance::GetModule() const
 {
-	return CastChecked<UAnimNextModule>(Asset);
+	return CastChecked<UAnimNextModule>(Asset, ECastCheckedType::NullAllowed);
 }
 
 #if WITH_EDITOR
diff --git a/Engine/Plugins/Experimental/UAF/UAFAnimGraph/Source/UAFAnimGraph/Private/Graph/AnimNextAnimationGraph.cpp b/Engine/Plugins/Experimental/UAF/UAFAnimGraph/Source/UAFAnimGraph/Private/Graph/AnimNextAnimationGraph.cpp
--- a/Engine/Plugins/Experimental/UAF/UAFAnimGraph/Source/UAFAnimGraph/Private/Graph/AnimNextAnimationGraph.cpp
+++ b/Engine/Plugins/Experimental/UAF/UAFAnimGraph/Source/UAFAnimGraph/Private/Graph/AnimNextAnimationGraph.cpp
@@ -20,7 +20,8 @@
 #include "Module/AnimNextSkeletalMeshComponentReferenceComponent.h"
 #include "RigVMBlueprintGeneratedClass.h"
 
-#if WITH_EDITOR	
+#if WITH_EDITOR
+#include "Editor.h"
 #include "Engine/ExternalAssetDependencyGatherer.h"
 REGISTER_ASSETDEPENDENCY_GATHERER(FExternalAssetDependencyGatherer, UAnimNextAnimationGraph);
 #endif // WITH_EDITOR
@@ -33,6 +34,13 @@ UAnimNextAnimationGraph::UAnimNextAnimationGraph(const FObjectInitializer& Objec
 	: Super(ObjectInitializer)
 {
 	ExtendedExecuteContext.SetContextPublicDataStruct(FAnimNextExecuteContext::StaticStruct());
+
+#if WITH_EDITOR
+	if (!IsTemplate())
+	{
+		FEditorDelegates::OnPreForceDeleteObjects.AddUObject(this, &UAnimNextAnimationGraph::OnPreForceDeleteObjects);
+	}
+#endif
 }
 
 TSharedPtr<FAnimNextGraphInstance> UAnimNextAnimationGraph::AllocateInstance(const UE::UAF::FGraphAllocationParams& InParams) const
@@ -99,6 +107,18 @@ TSharedPtr<FAnimNextGraphInstance> UAnimNextAnimationGraph::AllocateInstance(con
 	return InstanceImpl;
 }
 
+#if WITH_EDITOR
+void UAnimNextAnimationGraph::BeginDestroy()
+{
+	Super::BeginDestroy();
+
+	if (!IsTemplate())
+	{
+		FEditorDelegates::OnPreForceDeleteObjects.RemoveAll(this);
+	}
+}
+#endif
+
 void UAnimNextAnimationGraph::Serialize(FArchive& Ar)
 {
 	Super::Serialize(Ar);
@@ -206,6 +226,25 @@ bool UAnimNextAnimationGraph::LoadFromArchiveBuffer(const TArray<uint8>& InShare
 	}
 }
 
+#if WITH_EDITOR
+void UAnimNextAnimationGraph::OnPreForceDeleteObjects(const TArray<UObject*>& ObjectsToDelete)
+{
+	if (!ObjectsToDelete.Contains(this))
+	{
+		// Not deleting this graph instance
+		return;
+	}
+
+	FScopeLock Lock(&GraphInstancesLock);
+
+	TSet<FAnimNextGraphInstance*> GraphInstancesCopy = GraphInstances;
+	for (FAnimNextGraphInstance* GraphInstance : GraphInstancesCopy)
+	{
+		GraphInstance->Release();
+	}
+}
+#endif
+
 #if WITH_EDITORONLY_DATA
 void UAnimNextAnimationGraph::FreezeGraphInstances()
 {
diff --git a/Engine/Plugins/Experimental/UAF/UAFAnimGraph/Source/UAFAnimGraph/Private/Graph/AnimNextGraphInstance.cpp b/Engine/Plugins/Experimental/UAF/UAFAnimGraph/Source/UAFAnimGraph/Private/Graph/AnimNextGraphInstance.cpp
--- a/Engine/Plugins/Experimental/UAF/UAFAnimGraph/Source/UAFAnimGraph/Private/Graph/AnimNextGraphInstance.cpp
+++ b/Engine/Plugins/Experimental/UAF/UAFAnimGraph/Source/UAFAnimGraph/Private/Graph/AnimNextGraphInstance.cpp
@@ -59,7 +59,7 @@ bool FAnimNextGraphInstance::IsValid() const
 
 const UAnimNextAnimationGraph* FAnimNextGraphInstance::GetAnimationGraph() const
 {
-	return CastChecked<UAnimNextAnimationGraph>(Asset);
+	return CastChecked<UAnimNextAnimationGraph>(Asset, ECastCheckedType::NullAllowed);
 }
 
 FName FAnimNextGraphInstance::GetEntryPoint() const
diff --git a/Engine/Plugins/Experimental/UAF/UAFAnimGraph/Source/UAFAnimGraph/Public/Graph/AnimNextAnimationGraph.h b/Engine/Plugins/Experimental/UAF/UAFAnimGraph/Source/UAFAnimGraph/Public/Graph/AnimNextAnimationGraph.h
--- a/Engine/Plugins/Experimental/UAF/UAFAnimGraph/Source/UAFAnimGraph/Public/Graph/AnimNextAnimationGraph.h
+++ b/Engine/Plugins/Experimental/UAF/UAFAnimGraph/Source/UAFAnimGraph/Public/Graph/AnimNextAnimationGraph.h
@@ -94,6 +94,9 @@ public:
 
 	// UObject interface
 	UE_API virtual void Serialize(FArchive& Ar) override;
+#if WITH_EDITOR
+	UE_API virtual void BeginDestroy() override;
+#endif
 #if WITH_EDITORONLY_DATA
 	UE_API virtual void GetPreloadDependencies(TArray<UObject*>& OutDeps) override;
 #endif
@@ -107,6 +110,10 @@ private:
 	// Loads the graph data from the provided archive buffer and returns true on success, false otherwise
 	UE_API bool LoadFromArchiveBuffer(const TArray<uint8>& SharedDataArchiveBuffer);
 
+#if WITH_EDITOR
+	void OnPreForceDeleteObjects(const TArray<UObject*>& ObjectsToDelete);
+#endif
+
 #if WITH_EDITORONLY_DATA
 	// During graph compilation, if we have existing graph instances, we freeze them by releasing their memory before thawing them
 	// Freezing is a partial release of resources that retains the necessary information to re-create things safely
diff --git a/Engine/Plugins/Experimental/UAF/UAFAnimGraph/Source/UAFAnimGraph/UAFAnimGraph.Build.cs b/Engine/Plugins/Experimental/UAF/UAFAnimGraph/Source/UAFAnimGraph/UAFAnimGraph.Build.cs
--- a/Engine/Plugins/Experimental/UAF/UAFAnimGraph/Source/UAFAnimGraph/UAFAnimGraph.Build.cs
+++ b/Engine/Plugins/Experimental/UAF/UAFAnimGraph/Source/UAFAnimGraph/UAFAnimGraph.Build.cs
@@ -32,6 +32,7 @@ namespace UnrealBuildTool.Rules
 				PrivateDependencyModuleNames.AddRange(
 					new string[]
 					{
+						"UnrealEd",
 						"AssetRegistry",
 						"RigVMDeveloper",
 					}
