Add scoped named event for UAF pose evaluation
Rename UAF named events

diff --git a/Engine/Plugins/Experimental/MovieSceneAnimMixer/Source/MovieSceneAnimMixer/Private/Systems/MovieSceneRootMotionSystem.cpp b/Engine/Plugins/Experimental/MovieSceneAnimMixer/Source/MovieSceneAnimMixer/Private/Systems/MovieSceneRootMotionSystem.cpp
--- a/Engine/Plugins/Experimental/MovieSceneAnimMixer/Source/MovieSceneAnimMixer/Private/Systems/MovieSceneRootMotionSystem.cpp
+++ b/Engine/Plugins/Experimental/MovieSceneAnimMixer/Source/MovieSceneAnimMixer/Private/Systems/MovieSceneRootMotionSystem.cpp
@@ -484,7 +484,7 @@ void FAnimNextStoreRootTransformTask::Execute(UE::UAF::FEvaluationVM& VM) const
 
 		auto ApplyRootMotion = [RootTransform, WeakTarget = RootMotionData->Target]
 		{
-			SCOPED_NAMED_EVENT(AnimNextStoreRootTransformTask_ApplyRootMotion, FColor::Orange);
+			SCOPED_NAMED_EVENT(UAF_StoreRootTransformTask_ApplyRootMotion, FColor::Orange);
 			USceneComponent* TargetComponent = WeakTarget.Get();
 			if (TargetComponent && !GActorMovementTracker.IsActive())
 			{
diff --git a/Engine/Plugins/Experimental/UAF/UAF/Source/UAF/Private/Component/SkinnedMeshComponentExtensions.cpp b/Engine/Plugins/Experimental/UAF/UAF/Source/UAF/Private/Component/SkinnedMeshComponentExtensions.cpp
--- a/Engine/Plugins/Experimental/UAF/UAF/Source/UAF/Private/Component/SkinnedMeshComponentExtensions.cpp
+++ b/Engine/Plugins/Experimental/UAF/UAF/Source/UAF/Private/Component/SkinnedMeshComponentExtensions.cpp
@@ -148,7 +148,7 @@ void FSkinnedMeshComponentExtensions::CompleteAndDispatch(
 					return;
 				}
 
-				SCOPED_NAMED_EVENT(AnimNext_SkinnedMesh_CompleteAndDispatch_GameThread, FColor::Orange);
+				SCOPED_NAMED_EVENT(UAF_SkinnedMesh_CompleteAndDispatch_GameThread, FColor::Orange);
 				if (bNeedsChildTransformUpdate)
 				{
 					Component->UpdateChildTransforms(EUpdateTransformFlags::OnlyUpdateIfUsingSocket);
diff --git a/Engine/Plugins/Experimental/UAF/UAF/Source/UAF/Private/DataRegistry.cpp b/Engine/Plugins/Experimental/UAF/UAF/Source/UAF/Private/DataRegistry.cpp
--- a/Engine/Plugins/Experimental/UAF/UAF/Source/UAF/Private/DataRegistry.cpp
+++ b/Engine/Plugins/Experimental/UAF/UAF/Source/UAF/Private/DataRegistry.cpp
@@ -75,7 +75,7 @@ FDataRegistry* FDataRegistry::Get()
 
 /*static*/ void FDataRegistry::HandlePostGarbageCollect()
 {
-	SCOPED_NAMED_EVENT(AnimNext_DataRegistry_PostGC, FColor::Blue);
+	SCOPED_NAMED_EVENT(UAF_DataRegistry_PostGC, FColor::Blue);
 
 	// Compact the registry on GC
 	if (GAnimationDataRegistry)
@@ -123,7 +123,7 @@ bool FDataRegistry::GCCleanupTicker(float DeltaTime)
 
 bool FDataRegistry::RunGCCleanup(double TimeBudget)
 {
-	SCOPED_NAMED_EVENT(AnimNext_DataRegistry_GCCleanup, FColor::Blue);
+	SCOPED_NAMED_EVENT(UAF_DataRegistry_GCCleanup, FColor::Blue);
 
 	if (TimeBudget <= 0.0)
 	{
diff --git a/Engine/Plugins/Experimental/UAF/UAF/Source/UAF/Private/Module/ModuleTickFunction.cpp b/Engine/Plugins/Experimental/UAF/UAF/Source/UAF/Private/Module/ModuleTickFunction.cpp
--- a/Engine/Plugins/Experimental/UAF/UAF/Source/UAF/Private/Module/ModuleTickFunction.cpp
+++ b/Engine/Plugins/Experimental/UAF/UAF/Source/UAF/Private/Module/ModuleTickFunction.cpp
@@ -44,7 +44,7 @@ void FModuleEventTickFunction::ExecuteBindings_WT(float InDeltaTime)
 
 void FModuleEventTickFunction::EndTick(float DeltaTime)
 {
-	SCOPED_NAMED_EVENT(AnimNext_Module_EndTick, FColor::Orange);
+	SCOPED_NAMED_EVENT(UAF_Module_EndTick, FColor::Orange);
 
 	// Give module a chance to finish up processing
 	ModuleInstance->EndExecution(DeltaTime);
@@ -109,7 +109,7 @@ void FModuleEventTickFunction::EndTick(float DeltaTime)
 	{
 		RunTaskOnGameThread([MainThreadActionEventList = MoveTemp(MainThreadActionEventList)]()
 			{
-				SCOPED_NAMED_EVENT(AnimNext_Module_EndTick_GameThread, FColor::Orange);
+				SCOPED_NAMED_EVENT(UAF_Module_EndTick_GameThread, FColor::Orange);
 				check(IsInGameThread());
 				for (const FAnimNextTraitEventPtr& Event : MainThreadActionEventList)
 				{
@@ -156,7 +156,7 @@ void FModuleEventTickFunction::ExecuteTick(float DeltaTime, ELevelTick TickType,
 	{
 		auto PostExecuteGameThreadTask = [PostExecuteGameThreadTaskList = MoveTemp(PostExecuteGameThreadTaskList)]()
 			{
-				SCOPED_NAMED_EVENT(AnimNext_Module_PostEventTick, FColor::Orange);
+				SCOPED_NAMED_EVENT(UAF_Module_PostEventTick, FColor::Orange);
 				check(IsInGameThread());
 
 				for (auto& Function : PostExecuteGameThreadTaskList)
@@ -186,7 +186,7 @@ void FModuleEventTickFunction::ExecuteTick(float DeltaTime, ELevelTick TickType,
 
 void FModuleEventTickFunction::Run(float InDeltaTime)
 {
-	SCOPED_NAMED_EVENT(AnimNext_Module_EventTick, FColor::Orange);
+	SCOPED_NAMED_EVENT(UAF_Module_EventTick, FColor::Orange);
 
 	while (!PreExecuteTasks.IsEmpty())
 	{
diff --git a/Engine/Plugins/Experimental/UAF/UAFAnimGraph/Source/UAFAnimGraph/Private/Graph/RigUnit_AnimNextRunAnimationGraph_v1.cpp b/Engine/Plugins/Experimental/UAF/UAFAnimGraph/Source/UAFAnimGraph/Private/Graph/RigUnit_AnimNextRunAnimationGraph_v1.cpp
--- a/Engine/Plugins/Experimental/UAF/UAFAnimGraph/Source/UAFAnimGraph/Private/Graph/RigUnit_AnimNextRunAnimationGraph_v1.cpp
+++ b/Engine/Plugins/Experimental/UAF/UAFAnimGraph/Source/UAFAnimGraph/Private/Graph/RigUnit_AnimNextRunAnimationGraph_v1.cpp
@@ -18,7 +18,7 @@
 
 FRigUnit_AnimNextRunAnimationGraph_v1_Execute()
 {
-	SCOPED_NAMED_EVENT(AnimNext_Run_Graph_V1, FColor::Orange);
+	SCOPED_NAMED_EVENT(UAF_Run_Graph_V1, FColor::Orange);
 
 	using namespace UE::UAF;
 
diff --git a/Engine/Plugins/Experimental/UAF/UAFAnimGraph/Source/UAFAnimGraph/Private/Graph/RigUnit_AnimNextRunAnimationGraph_v2.cpp b/Engine/Plugins/Experimental/UAF/UAFAnimGraph/Source/UAFAnimGraph/Private/Graph/RigUnit_AnimNextRunAnimationGraph_v2.cpp
--- a/Engine/Plugins/Experimental/UAF/UAFAnimGraph/Source/UAFAnimGraph/Private/Graph/RigUnit_AnimNextRunAnimationGraph_v2.cpp
+++ b/Engine/Plugins/Experimental/UAF/UAFAnimGraph/Source/UAFAnimGraph/Private/Graph/RigUnit_AnimNextRunAnimationGraph_v2.cpp
@@ -22,7 +22,7 @@
 
 FRigUnit_AnimNextRunAnimationGraph_v2_Execute()
 {
-	SCOPED_NAMED_EVENT(AnimNext_Run_Graph_V2, FColor::Orange);
+	SCOPED_NAMED_EVENT(UAF_Run_Graph_V2, FColor::Orange);
 
 	using namespace UE::UAF;
 
diff --git a/Engine/Plugins/Experimental/UAF/UAFAnimGraph/Source/UAFAnimGraph/Private/TraitInterfaces/IEvaluate.cpp b/Engine/Plugins/Experimental/UAF/UAFAnimGraph/Source/UAFAnimGraph/Private/TraitInterfaces/IEvaluate.cpp
--- a/Engine/Plugins/Experimental/UAF/UAFAnimGraph/Source/UAFAnimGraph/Private/TraitInterfaces/IEvaluate.cpp
+++ b/Engine/Plugins/Experimental/UAF/UAFAnimGraph/Source/UAFAnimGraph/Private/TraitInterfaces/IEvaluate.cpp
@@ -98,7 +98,7 @@ namespace UE::UAF
 
 	FEvaluationProgram EvaluateGraph(const FEvaluateGraphContext& EvaluateGraphContext, const FWeakTraitPtr& GraphRootPtr)
 	{
-		SCOPED_NAMED_EVENT(AnimNext_EvaluateGraph, FColor::Orange);
+		SCOPED_NAMED_EVENT(UAF_EvaluateGraph, FColor::Orange);
 		
 		FEvaluationProgram EvaluationProgram;
 
@@ -232,30 +232,34 @@ namespace UE::UAF
 
 		TRACE_ANIMNEXT_EVALUATIONPROGRAM(EvaluationProgram, GraphInstance);
 
-		FEvaluationVM EvaluationVM(EEvaluationFlags::All, RefPose, GraphLODLevel);
-		bool bHasValidOutput = false;
-
-		if (!EvaluationProgram.IsEmpty())
 		{
-			EvaluationProgram.Execute(EvaluationVM);
+			SCOPED_NAMED_EVENT(UAF_EvaluatePose, FColor::Orange);
+
+			FEvaluationVM EvaluationVM(EEvaluationFlags::All, RefPose, GraphLODLevel);
+			bool bHasValidOutput = false;
 
-			TUniquePtr<FKeyframeState> EvaluatedKeyframe;
-			if (EvaluationVM.PopValue(KEYFRAME_STACK_NAME, EvaluatedKeyframe))
+			if (!EvaluationProgram.IsEmpty())
 			{
-				OutputPose.LODPose.CopyFrom(EvaluatedKeyframe->Pose);
-				OutputPose.Curves.CopyFrom(EvaluatedKeyframe->Curves);
-				OutputPose.Attributes.CopyFrom(EvaluatedKeyframe->Attributes);
-				bHasValidOutput = true;
+				EvaluationProgram.Execute(EvaluationVM);
+
+				TUniquePtr<FKeyframeState> EvaluatedKeyframe;
+				if (EvaluationVM.PopValue(KEYFRAME_STACK_NAME, EvaluatedKeyframe))
+				{
+					OutputPose.LODPose.CopyFrom(EvaluatedKeyframe->Pose);
+					OutputPose.Curves.CopyFrom(EvaluatedKeyframe->Curves);
+					OutputPose.Attributes.CopyFrom(EvaluatedKeyframe->Attributes);
+					bHasValidOutput = true;
+				}
 			}
-		}
 
-		if (!bHasValidOutput)
-		{
-			// We need to output a valid pose, generate one
-			FKeyframeState ReferenceKeyframe = EvaluationVM.MakeReferenceKeyframe(false);
-			OutputPose.LODPose.CopyFrom(ReferenceKeyframe.Pose);
-			OutputPose.Curves.CopyFrom(ReferenceKeyframe.Curves);
-			OutputPose.Attributes.CopyFrom(ReferenceKeyframe.Attributes);
+			if (!bHasValidOutput)
+			{
+				// We need to output a valid pose, generate one
+				FKeyframeState ReferenceKeyframe = EvaluationVM.MakeReferenceKeyframe(false);
+				OutputPose.LODPose.CopyFrom(ReferenceKeyframe.Pose);
+				OutputPose.Curves.CopyFrom(ReferenceKeyframe.Curves);
+				OutputPose.Attributes.CopyFrom(ReferenceKeyframe.Attributes);
+			}
 		}
 	}
 }
diff --git a/Engine/Plugins/Experimental/UAF/UAFAnimGraph/Source/UAFAnimGraph/Private/TraitInterfaces/IUpdate.cpp b/Engine/Plugins/Experimental/UAF/UAFAnimGraph/Source/UAFAnimGraph/Private/TraitInterfaces/IUpdate.cpp
--- a/Engine/Plugins/Experimental/UAF/UAFAnimGraph/Source/UAFAnimGraph/Private/TraitInterfaces/IUpdate.cpp
+++ b/Engine/Plugins/Experimental/UAF/UAFAnimGraph/Source/UAFAnimGraph/Private/TraitInterfaces/IUpdate.cpp
@@ -474,7 +474,7 @@ namespace UE::UAF
 
 	void UpdateGraph(FUpdateGraphContext& UpdateGraphContext)
 	{
-		SCOPED_NAMED_EVENT(AnimNext_UpdateGraph, FColor::Orange);
+		SCOPED_NAMED_EVENT(UAF_UpdateGraph, FColor::Orange);
 		
 		FAnimNextGraphInstance& GraphInstance = UpdateGraphContext.GetGraphInstance(); 
 
