Avoid crash when opening the context menu on legacy UAF content

diff --git a/Engine/Plugins/Experimental/UAF/UAFAnimGraph/Source/UAFAnimGraphEditor/Private/Graph/AnimNextAnimationGraphMenuExtensions.cpp b/Engine/Plugins/Experimental/UAF/UAFAnimGraph/Source/UAFAnimGraphEditor/Private/Graph/AnimNextAnimationGraphMenuExtensions.cpp
--- a/Engine/Plugins/Experimental/UAF/UAFAnimGraph/Source/UAFAnimGraphEditor/Private/Graph/AnimNextAnimationGraphMenuExtensions.cpp
+++ b/Engine/Plugins/Experimental/UAF/UAFAnimGraph/Source/UAFAnimGraphEditor/Private/Graph/AnimNextAnimationGraphMenuExtensions.cpp
@@ -97,7 +97,8 @@ void FAnimationGraphMenuExtensions::RegisterMenus()
 		
 		if (UncookedOnly::FAnimGraphUtils::IsTraitStackNode(ModelNode))
 		{
-			UAnimNextTraitStackUnitNode* TraitStackNode = CastChecked<UAnimNextTraitStackUnitNode>(ModelNode);
+			// Cast test to deal with legacy content
+			if (UAnimNextTraitStackUnitNode* TraitStackNode = Cast<UAnimNextTraitStackUnitNode>(ModelNode))
 			{
 				FToolMenuSection& Section = InMenu->AddSection("AnimNextTemplateNodeActions", LOCTEXT("AnimNextTemplateNodeActionsMenuHeader", "Template"));
 
