Temporary fix to ensure that UAF graph builder hashes all relevant data

diff --git a/Engine/Plugins/Experimental/UAF/UAFAnimGraph/Source/UAFAnimGraph/Private/Factory/AnimNextSimpleAnimGraphBuilder.cpp b/Engine/Plugins/Experimental/UAF/UAFAnimGraph/Source/UAFAnimGraph/Private/Factory/AnimNextSimpleAnimGraphBuilder.cpp
--- a/Engine/Plugins/Experimental/UAF/UAFAnimGraph/Source/UAFAnimGraph/Private/Factory/AnimNextSimpleAnimGraphBuilder.cpp
+++ b/Engine/Plugins/Experimental/UAF/UAFAnimGraph/Source/UAFAnimGraph/Private/Factory/AnimNextSimpleAnimGraphBuilder.cpp
@@ -11,6 +11,16 @@
 #include "TraitCore/TraitWriter.h"
 #include "Traits/ReferencePoseTrait.h"
 
+// [Temporary Hashing Fix]
+// FArchiveCityHash64 does not currently specify any flags (e.g. saving, loading, counting memory) and as such
+// can end up missing content. In particular, FInstancedStruct instances are entirely skipped and not factored
+// into the resulting hash which is problematic and incorrect.
+// Due to time constraints and the risky nature of performing an Engine fix at this time, a temporary fix is
+// made in this file to enable the counting memory flag which resolves the issue with instanced structs which
+// we rely on.
+#include "Hash/CityHash.h"
+#include "Serialization/ArchiveUObject.h"
+
 #include UE_INLINE_GENERATED_CPP_BY_NAME(AnimNextSimpleAnimGraphBuilder)
 
 bool FAnimNextSimpleAnimGraphBuilder::Build(UE::UAF::FAnimGraphBuilderContext& InContext) const
@@ -217,7 +227,107 @@ void FAnimNextSimpleAnimGraphBuilder::PushTraitInstancedStructToStack(int32 InSt
 	InvalidateKey();
 }
 
+namespace UE::UAF::Private
+{
+	// See [Temporary Hashing Fix] above for details
+
+	class FArchiveCityHash64_Tmp : public FArchiveUObject
+	{
+	public:
+		explicit FArchiveCityHash64_Tmp(const uint64 InHash = 0)
+			: Hash(InHash)
+		{
+			// Without this flag (or some other), FInstancedStructs aren't traversed and serialized by this archive
+			ArIsCountingMemory = true;
+		}
+
+		/**
+		 * @return The hash computed so far.
+		 */
+		uint64 GetHash() const { return Hash; }
+
+		//~ Begin FArchive Interface
+		virtual void Serialize(void* Data, int64 Num) override
+		{
+			const char* BytePointer = static_cast<const char*>(Data);
+			while (Num > 0)
+			{
+				const int32 BytesToHash = static_cast<int32>(FMath::Min(Num, static_cast<int64>(MAX_int32)));
+
+				Hash = CityHash64WithSeed(BytePointer, BytesToHash, Hash);
+
+				Num -= BytesToHash;
+				BytePointer += BytesToHash;
+			}
+		}
+
+		virtual FArchive& operator<<(class FName& Name) override
+		{
+			Hash = HashString(Name.ToString(), Hash);
+			return *this;
+		}
+
+		virtual FArchive& operator<<(class UObject*& Object) override
+		{
+			Hash = HashString(GetPathNameSafe(Object), Hash);
+			return *this;
+		}
+
+		virtual FString GetArchiveName() const override { return TEXT("FArchiveCityHash64_Tmp"); }
+		//~ End FArchive Interface
+
+		FORCEINLINE static uint64 HashString(const FString& InString, const uint64 InHash = 0)
+		{
+			return CityHash64WithSeed(reinterpret_cast<const char*>(*InString), static_cast<uint32>(InString.GetAllocatedSize()), InHash);
+		}
+
+	private:
+		uint64 Hash;
+	};
+
+	uint64 GetStructHash64_Tmp(const UScriptStruct& ScriptStruct, const uint8* StructMemory)
+	{
+		const uint64 BaseHash = FArchiveCityHash64_Tmp::HashString(ScriptStruct.GetPathName());
+
+		const UScriptStruct::ICppStructOps* CppStructOps = ScriptStruct.GetCppStructOps();
+		if (CppStructOps && CppStructOps->HasGetTypeHash())
+		{
+			const uint32 StructHash = ScriptStruct.GetStructTypeHash(StructMemory);
+			return CityHash128to64(Uint128_64(BaseHash, StructHash));
+		}
+		else if (StructMemory)
+		{
+			FArchiveCityHash64_Tmp Ar(BaseHash);
+			UScriptStruct& NonConstScriptStruct = const_cast<UScriptStruct&>(ScriptStruct);
+			NonConstScriptStruct.SerializeItem(Ar, const_cast<uint8*>(StructMemory), nullptr);
+
+			return Ar.GetHash();
+		}
+
+		return BaseHash;
+	}
+
+	template<typename T>
+	uint64 GetStructHash64Helper_Tmp(const T& Struct)
+	{
+		if (const UScriptStruct* ScriptStruct = Struct.GetScriptStruct())
+		{
+			return GetStructHash64_Tmp(*ScriptStruct, Struct.GetMemory());
+		}
+		return 0;
+	}
+
+	uint64 GetStructHash64_Tmp(const FConstStructView& StructView)
+	{
+		return GetStructHash64Helper_Tmp(StructView);
+	}
+}
+
 uint64 FAnimNextSimpleAnimGraphBuilder::RecalculateKey() const
 {
-	return UE::StructUtils::GetStructHash64(TConstStructView<FAnimNextSimpleAnimGraphBuilder>(*this));
+	// See [Temporary Hashing Fix] above for details
+	return UE::UAF::Private::GetStructHash64_Tmp(TConstStructView<FAnimNextSimpleAnimGraphBuilder>(*this));
+
+	// Original code left intentionally here to be restored once the hashing is fixed in engine
+	//return UE::StructUtils::GetStructHash64(TConstStructView<FAnimNextSimpleAnimGraphBuilder>(*this));
 }
