AnimNextFactoryParamsDetails changes to make it work outside of the state tree editor

diff --git a/Engine/Plugins/Experimental/UAF/UAFAnimGraph/Source/UAFAnimGraphEditor/Private/Factory/AnimNextFactoryParamsDetails.cpp b/Engine/Plugins/Experimental/UAF/UAFAnimGraph/Source/UAFAnimGraphEditor/Private/Factory/AnimNextFactoryParamsDetails.cpp
--- a/Engine/Plugins/Experimental/UAF/UAFAnimGraph/Source/UAFAnimGraphEditor/Private/Factory/AnimNextFactoryParamsDetails.cpp
+++ b/Engine/Plugins/Experimental/UAF/UAFAnimGraph/Source/UAFAnimGraphEditor/Private/Factory/AnimNextFactoryParamsDetails.cpp
@@ -8,8 +8,10 @@
 #include "IDetailChildrenBuilder.h"
 #include "IDetailPropertyRow.h"
 #include "InstancedStructDetails.h"
+#include "IPropertyAccessEditor.h"
 #include "IStructureDataProvider.h"
 #include "PropertyHandle.h"
+#include "ScopedTransaction.h"
 #include "StructViewerFilter.h"
 #include "StructViewerModule.h"
 #include "Algo/Compare.h"
@@ -56,6 +58,27 @@ void FAnimNextFactoryParamsDetails::CustomizeChildren(TSharedRef<IPropertyHandle
 	{
 		if (TSharedPtr<IPropertyHandle> PinnedPropertyHandle = WeakPropertyHandle.Pin())
 		{
+			FString FactorySourcePropertyName = PinnedPropertyHandle->GetMetaData("FactorySourceProperty");
+			TSharedPtr<IPropertyHandle> FactorySourceHandle = PinnedPropertyHandle->GetParentHandle()->GetChildHandle(FName(*FactorySourcePropertyName));
+
+			if (FactorySourceHandle)
+			{
+				UObject* FactorySourceObject = nullptr;
+				FAnimNextFactoryParams DefaultParams;
+				if (FactorySourceHandle->GetValue(FactorySourceObject) == FPropertyAccess::Result::Success && FactorySourceObject)
+				{
+					DefaultParams = UE::UAF::FAnimGraphFactory::GetDefaultParamsForObject(FactorySourceObject);
+				}
+				
+				PinnedPropertyHandle->EnumerateRawData([&DefaultParams](void* RawData, const int32 DataIndex, const int32 NumDatas)
+					{
+						// todo: ideally we should retain settings on parameter structs that are of the same type
+						*static_cast<FAnimNextFactoryParams*>(RawData) = DefaultParams;
+						return true;
+					}
+				);
+			}
+			
 			PinnedPropertyHandle->RequestRebuildChildren();
 		}
 	}));
@@ -82,11 +105,14 @@ void FAnimNextFactoryParamsDetails::CustomizeChildren(TSharedRef<IPropertyHandle
 			}
 			else
 			{
-				const FAnimNextFactoryParams& DefaultParams = FAnimGraphFactory::GetDefaultParamsForObject(Object);
-				if (!Algo::Compare(DefaultParams.Builder.Stacks[0].TraitStructs, CommonDefaults->Builder.Stacks[0].TraitStructs))
+				if (Object)
 				{
-					CommonDefaults.Reset();
-					return false;
+					const FAnimNextFactoryParams& DefaultParams = FAnimGraphFactory::GetDefaultParamsForObject(Object);
+					if (!Algo::Compare(DefaultParams.Builder.Stacks[0].TraitStructs, CommonDefaults->Builder.Stacks[0].TraitStructs))
+					{
+						CommonDefaults.Reset();
+						return false;
+					}
 				}
 			}
 			return true;
@@ -157,7 +183,12 @@ void FAnimNextFactoryParamsDetails::CustomizeChildren(TSharedRef<IPropertyHandle
 					AddedStructs.Add(CommonType.GetValue());
 
 					const FName UniqueName(InstancedStructHandle->GetPropertyPath());
-					IDetailPropertyRow* Row = InChildBuilder.AddChildStructure(InstancedStructHandle, MakeShared<FInstancedStructProvider>(InstancedStructHandle), UniqueName);
+					TSharedPtr<FInstancedStructProvider> StructProvider = MakeShared<FInstancedStructProvider>(InstancedStructHandle);
+					
+					if (StructProvider->GetBaseStructure() == nullptr)
+						continue;
+					
+					IDetailPropertyRow* Row = InChildBuilder.AddChildStructure(InstancedStructHandle, StructProvider, UniqueName);
 					FDetailWidgetRow& CustomWidget = Row->CustomWidget(true);
 
 					CustomWidget
@@ -253,6 +284,8 @@ void FAnimNextFactoryParamsDetails::CustomizeChildren(TSharedRef<IPropertyHandle
 						TSharedPtr<IPropertyHandle> PinnedPropertyHandle = WeakPropertyHandle.Pin();
 						if (PinnedTraitStructsHandle.IsValid() && PinnedPropertyHandle.IsValid())
 						{
+							FScopedTransaction Transaction(LOCTEXT("Enable Trait","Enable Trait"));
+							
 							TSharedPtr<IPropertyHandleArray> ArrayHandle = PinnedTraitStructsHandle->AsArray();
 							FPropertyHandleItemAddResult Result = ArrayHandle->AddItem();
 							TSharedPtr<IPropertyHandle> NewEntryHandle = PinnedTraitStructsHandle->GetChildHandle(Result.GetIndex());
@@ -348,6 +381,8 @@ void FAnimNextFactoryParamsDetails::CustomizeChildren(TSharedRef<IPropertyHandle
 							TSharedPtr<IPropertyHandle> PinnedPropertyHandle = WeakPropertyHandle.Pin();
 							if (PinnedPropertyHandle.IsValid())
 							{
+								FScopedTransaction Transaction(LOCTEXT("Add Trait","Add Trait"));
+								
 								TSharedPtr<IPropertyHandle> BuilderHandle = PinnedPropertyHandle->GetChildHandle(GET_MEMBER_NAME_CHECKED(FAnimNextFactoryParams, Builder));
 								TSharedPtr<IPropertyHandle> StacksHandle = BuilderHandle->GetChildHandle(GET_MEMBER_NAME_CHECKED(FAnimNextSimpleAnimGraphBuilder, Stacks));
 								TSharedPtr<IPropertyHandle> StackHandle = StacksHandle->GetChildHandle(0);
@@ -360,7 +395,7 @@ void FAnimNextFactoryParamsDetails::CustomizeChildren(TSharedRef<IPropertyHandle
 								TSharedPtr<IPropertyHandle> TraitStructsHandle = StackHandle->GetChildHandle(GET_MEMBER_NAME_CHECKED(FAnimNextSimpleAnimGraphBuilderTraitStackDesc, TraitStructs));
 								TSharedPtr<IPropertyHandleArray> ArrayHandle = TraitStructsHandle->AsArray();
 								FPropertyHandleItemAddResult Result = ArrayHandle->AddItem();
-								TSharedPtr<IPropertyHandle> NewEntryHandle = TraitStructsHandle->GetChildHandle(Result.GetIndex());
+								TSharedPtr<IPropertyHandle>  NewEntryHandle = TraitStructsHandle->GetChildHandle(Result.GetIndex());
 								NewEntryHandle->EnumerateRawData([InStruct](void* RawData, const int32 DataIndex, const int32 NumDatas)
 								{
 									FInstancedStruct* InstancedStruct = static_cast<FInstancedStruct*>(RawData);
diff --git a/Engine/Plugins/Experimental/UAF/UAFStateTree/Source/UAFStateTree/Internal/Tasks/AnimNextStateTreeGraphInstanceTask.h b/Engine/Plugins/Experimental/UAF/UAFStateTree/Source/UAFStateTree/Internal/Tasks/AnimNextStateTreeGraphInstanceTask.h
--- a/Engine/Plugins/Experimental/UAF/UAFStateTree/Source/UAFStateTree/Internal/Tasks/AnimNextStateTreeGraphInstanceTask.h
+++ b/Engine/Plugins/Experimental/UAF/UAFStateTree/Source/UAFStateTree/Internal/Tasks/AnimNextStateTreeGraphInstanceTask.h
@@ -102,7 +102,6 @@ protected:
 	virtual EStateTreeRunStatus Tick(FStateTreeExecutionContext& Context, const float DeltaTime) const override;
 	virtual void ExitState(FStateTreeExecutionContext& Context, const FStateTreeTransitionResult& Transition) const override;
 #if WITH_EDITOR
-	virtual void PostEditInstanceDataChangeChainProperty(const FPropertyChangedChainEvent& PropertyChangedEvent, FStateTreeDataView InstanceDataView) override;
 	virtual void GetObjectReferences(TArray<const UObject*>& OutReferencedObjects, const FStateTreeDataView InstanceDataView) const override;
 #endif
 public:
diff --git a/Engine/Plugins/Experimental/UAF/UAFStateTree/Source/UAFStateTree/Private/Tasks/AnimNextStateTreeGraphInstanceTask.cpp b/Engine/Plugins/Experimental/UAF/UAFStateTree/Source/UAFStateTree/Private/Tasks/AnimNextStateTreeGraphInstanceTask.cpp
--- a/Engine/Plugins/Experimental/UAF/UAFStateTree/Source/UAFStateTree/Private/Tasks/AnimNextStateTreeGraphInstanceTask.cpp
+++ b/Engine/Plugins/Experimental/UAF/UAFStateTree/Source/UAFStateTree/Private/Tasks/AnimNextStateTreeGraphInstanceTask.cpp
@@ -139,30 +139,6 @@ void FAnimNextStateTreeGraphInstanceTask::ExitState(FStateTreeExecutionContext&
 }
 
 #if WITH_EDITOR
-void FAnimNextStateTreeGraphInstanceTask::PostEditInstanceDataChangeChainProperty(const FPropertyChangedChainEvent& PropertyChangedEvent, FStateTreeDataView InstanceDataView)
-{
-	const FProperty* Property = PropertyChangedEvent.Property;
-	if(Property == nullptr)
-	{
-		return;
-	}
-
-	if(Property->GetFName() == GET_MEMBER_NAME_CHECKED(FAnimNextGraphInstanceTaskInstanceData, Asset))
-	{
-		// Repopulate payload property bag if the selected asset changes
-		FAnimNextGraphInstanceTaskInstanceData& InstanceData = InstanceDataView.GetMutable<FAnimNextGraphInstanceTaskInstanceData>();
-
-		if (InstanceData.Asset == nullptr || InstanceData.Asset.IsA<UAnimNextAnimationGraph>())
-		{
-			InstanceData.Parameters.Reset();
-		}
-		else
-		{
-			// TODO: Copy over matching values before overwrite to avoid losing user edits to matching trait data
-			InstanceData.Parameters = UE::UAF::FAnimGraphFactory::GetDefaultParamsForObject(InstanceData.Asset);
-		}
-	}
-}
 
 void FAnimNextStateTreeGraphInstanceTask::GetObjectReferences(TArray<const UObject*>& OutReferencedObjects, FStateTreeDataView InstanceDataView) const
 {
