UAF: Dont hook PreForceDelete delegate for transient (i.e. procedural) graphs

These graphs cant be deleted by users and registering the delegate can crash as it isnt thread safe

diff --git a/Engine/Plugins/Experimental/UAF/UAFAnimGraph/Source/UAFAnimGraph/Private/Graph/AnimNextAnimationGraph.cpp b/Engine/Plugins/Experimental/UAF/UAFAnimGraph/Source/UAFAnimGraph/Private/Graph/AnimNextAnimationGraph.cpp
--- a/Engine/Plugins/Experimental/UAF/UAFAnimGraph/Source/UAFAnimGraph/Private/Graph/AnimNextAnimationGraph.cpp
+++ b/Engine/Plugins/Experimental/UAF/UAFAnimGraph/Source/UAFAnimGraph/Private/Graph/AnimNextAnimationGraph.cpp
@@ -36,7 +36,7 @@ UAnimNextAnimationGraph::UAnimNextAnimationGraph(const FObjectInitializer& Objec
 	ExtendedExecuteContext.SetContextPublicDataStruct(FAnimNextExecuteContext::StaticStruct());
 
 #if WITH_EDITOR
-	if (!IsTemplate())
+	if (!IsTemplate() && !HasAnyFlags(RF_Transient))
 	{
 		FEditorDelegates::OnPreForceDeleteObjects.AddUObject(this, &UAnimNextAnimationGraph::OnPreForceDeleteObjects);
 	}
@@ -112,7 +112,7 @@ void UAnimNextAnimationGraph::BeginDestroy()
 {
 	Super::BeginDestroy();
 
-	if (!IsTemplate())
+	if (!IsTemplate() && !HasAnyFlags(RF_Transient))
 	{
 		FEditorDelegates::OnPreForceDeleteObjects.RemoveAll(this);
 	}
