UAF: Fix monolithic builds

diff --git a/Engine/Plugins/Experimental/UAF/UAF/Source/UAF/Private/Factory/SystemFactory.cpp b/Engine/Plugins/Experimental/UAF/UAF/Source/UAF/Private/Factory/SystemFactory.cpp
--- a/Engine/Plugins/Experimental/UAF/UAF/Source/UAF/Private/Factory/SystemFactory.cpp
+++ b/Engine/Plugins/Experimental/UAF/UAF/Source/UAF/Private/Factory/SystemFactory.cpp
@@ -20,7 +20,7 @@ static TStripedMap<32, uint64, TWeakObjectPtr<const UAnimNextModule>> GSystemMap
 static TMap<const UClass*, TPair<FUAFSystemFactoryParams, FSystemFactory::FParamsInitializer>> GDefaultParamsForObject;
 static TArray<UClass*> GAllFactoryClasses;
 
-const TPair<FUAFSystemFactoryParams, FSystemFactory::FParamsInitializer>* FindParamsForClass(const UClass* InClass)
+const TPair<FUAFSystemFactoryParams, FSystemFactory::FParamsInitializer>* FindSystemParamsForClass(const UClass* InClass)
 {
 	const UClass* CurrentClass = InClass;
 	const TPair<FUAFSystemFactoryParams, FSystemFactory::FParamsInitializer>* FoundParams = nullptr;
@@ -64,7 +64,7 @@ const UAnimNextModule* FSystemFactory::BuildSystem(const ISystemBuilder& InBuild
 	TWeakObjectPtr<const UAnimNextModule> System;
 	auto ApplyWeakSystem = [&ProduceSystem, &System](TWeakObjectPtr<const UAnimNextModule>& InFoundSystem)
 	{
-		if (InFoundSystem.IsStale())
+		if (InFoundSystem.IsStale() || InFoundSystem.Get() == nullptr)
 		{
 			// Found a stale weak ptr, so regenerate (this may be a procedural system or a previously loaded system that has been GCed)
 			InFoundSystem = ProduceSystem();
@@ -79,7 +79,7 @@ const UAnimNextModule* FSystemFactory::BuildSystem(const ISystemBuilder& InBuild
 const FUAFSystemFactoryParams& FSystemFactory::GetDefaultParamsForClass(const UClass* InClass)
 {
 	check(InClass);
-	const TPair<FUAFSystemFactoryParams, FParamsInitializer>* ParamsTask = Private::FindParamsForClass(InClass);
+	const TPair<FUAFSystemFactoryParams, FParamsInitializer>* ParamsTask = Private::FindSystemParamsForClass(InClass);
 	check(ParamsTask != nullptr);	// Unregistered object type
 	return ParamsTask->Key;
 }
@@ -87,7 +87,7 @@ const FUAFSystemFactoryParams& FSystemFactory::GetDefaultParamsForClass(const UC
 FUAFSystemFactoryParams FSystemFactory::GetDefaultParamsForObject(const UObject* InObject)
 {
 	check(InObject);
-	const TPair<FUAFSystemFactoryParams, FParamsInitializer>* ParamsTask = Private::FindParamsForClass(InObject->GetClass());
+	const TPair<FUAFSystemFactoryParams, FParamsInitializer>* ParamsTask = Private::FindSystemParamsForClass(InObject->GetClass());
 	check(ParamsTask != nullptr);	// Unregistered object type
 
 	FUAFSystemFactoryParams ParamsCopy = ParamsTask->Key;
diff --git a/Engine/Plugins/Experimental/UAF/UAFAnimGraph/Source/UAFAnimGraph/Private/Factory/AnimGraphFactory.cpp b/Engine/Plugins/Experimental/UAF/UAFAnimGraph/Source/UAFAnimGraph/Private/Factory/AnimGraphFactory.cpp
--- a/Engine/Plugins/Experimental/UAF/UAFAnimGraph/Source/UAFAnimGraph/Private/Factory/AnimGraphFactory.cpp
+++ b/Engine/Plugins/Experimental/UAF/UAFAnimGraph/Source/UAFAnimGraph/Private/Factory/AnimGraphFactory.cpp
@@ -21,7 +21,7 @@ static TStripedMap<32, uint64, TWeakObjectPtr<const UAnimNextAnimationGraph>> GA
 static TMap<const UClass*, TPair<FAnimNextFactoryParams, FAnimGraphFactory::FParamsInitializer>> GDefaultParamsForObject;
 static TArray<UClass*> GAllFactoryClasses;
 
-const TPair<FAnimNextFactoryParams, FAnimGraphFactory::FParamsInitializer>* FindParamsForClass(const UClass* InClass)
+const TPair<FAnimNextFactoryParams, FAnimGraphFactory::FParamsInitializer>* FindGraphParamsForClass(const UClass* InClass)
 {
 	const UClass* CurrentClass = InClass;
 	const TPair<FAnimNextFactoryParams, FAnimGraphFactory::FParamsInitializer>* FoundParams = nullptr;
@@ -46,7 +46,7 @@ uint64 IAnimGraphBuilder::GetKey() const
 
 const UAnimNextAnimationGraph* FAnimGraphFactory::GetDefaultGraphHost()
 {
-	const TPair<FAnimNextFactoryParams, FParamsInitializer>* ParamsTask = Private::FindParamsForClass(UAnimNextAnimationGraph::StaticClass());
+	const TPair<FAnimNextFactoryParams, FParamsInitializer>* ParamsTask = Private::FindGraphParamsForClass(UAnimNextAnimationGraph::StaticClass());
 	return BuildGraph(ParamsTask->Key.GetBuilder());
 }
 
@@ -96,7 +96,7 @@ const UAnimNextAnimationGraph* FAnimGraphFactory::BuildGraph(const IAnimGraphBui
 const FAnimNextFactoryParams& FAnimGraphFactory::GetDefaultParamsForClass(const UClass* InClass)
 {
 	check(InClass);
-	const TPair<FAnimNextFactoryParams, FParamsInitializer>* ParamsTask = Private::FindParamsForClass(InClass);
+	const TPair<FAnimNextFactoryParams, FParamsInitializer>* ParamsTask = Private::FindGraphParamsForClass(InClass);
 	check(ParamsTask != nullptr);	// Unregistered object type
 	return ParamsTask->Key;
 }
@@ -104,7 +104,7 @@ const FAnimNextFactoryParams& FAnimGraphFactory::GetDefaultParamsForClass(const
 FAnimNextFactoryParams FAnimGraphFactory::GetDefaultParamsForObject(const UObject* InObject)
 {
 	check(InObject);
-	const TPair<FAnimNextFactoryParams, FParamsInitializer>* ParamsTask = Private::FindParamsForClass(InObject->GetClass());
+	const TPair<FAnimNextFactoryParams, FParamsInitializer>* ParamsTask = Private::FindGraphParamsForClass(InObject->GetClass());
 	check(ParamsTask != nullptr);	// Unregistered object type
 
 	FAnimNextFactoryParams ParamsCopy = ParamsTask->Key;
@@ -121,7 +121,7 @@ FAnimNextFactoryParams FAnimGraphFactory::GetDefaultParamsForObject(const UObjec
 void FAnimGraphFactory::InitializeDefaultParamsForObject(const UObject* InObject, FAnimNextFactoryParams& InOutParams)
 {
 	check(InObject);
-	const TPair<FAnimNextFactoryParams, FParamsInitializer>* ParamsTask = Private::FindParamsForClass(InObject->GetClass());
+	const TPair<FAnimNextFactoryParams, FParamsInitializer>* ParamsTask = Private::FindGraphParamsForClass(InObject->GetClass());
 	check(ParamsTask != nullptr);	// Unregistered object type
 
 	if (InOutParams.Builder.Stacks.Num() > 0)
