UAF: Fix crash when creating a BP variable of factory params type

diff --git a/Engine/Plugins/Experimental/UAF/UAFAnimGraph/Source/UAFAnimGraphEditor/Private/Factory/AnimNextFactoryParamsDetails.cpp b/Engine/Plugins/Experimental/UAF/UAFAnimGraph/Source/UAFAnimGraphEditor/Private/Factory/AnimNextFactoryParamsDetails.cpp
--- a/Engine/Plugins/Experimental/UAF/UAFAnimGraph/Source/UAFAnimGraphEditor/Private/Factory/AnimNextFactoryParamsDetails.cpp
+++ b/Engine/Plugins/Experimental/UAF/UAFAnimGraph/Source/UAFAnimGraphEditor/Private/Factory/AnimNextFactoryParamsDetails.cpp
@@ -54,69 +54,73 @@ void FAnimNextFactoryParamsDetails::CustomizeChildren(TSharedRef<IPropertyHandle
 	TOptional<FAnimNextFactoryParams> CommonDefaults;
 	FString FactorySourcePropertyName = InPropertyHandle->GetMetaData("FactorySourceProperty");
 	TSharedPtr<IPropertyHandle> FactorySourceHandle = InPropertyHandle->GetParentHandle()->GetChildHandle(FName(*FactorySourcePropertyName));
-	FactorySourceHandle->SetOnPropertyValueChanged(FSimpleDelegate::CreateLambda([WeakPropertyHandle = TWeakPtr<IPropertyHandle>(InPropertyHandle)]()
+
+	if (FactorySourceHandle.IsValid())
 	{
-		if (TSharedPtr<IPropertyHandle> PinnedPropertyHandle = WeakPropertyHandle.Pin())
+		FactorySourceHandle->SetOnPropertyValueChanged(FSimpleDelegate::CreateLambda([WeakPropertyHandle = TWeakPtr<IPropertyHandle>(InPropertyHandle)]()
 		{
-			FString FactorySourcePropertyName = PinnedPropertyHandle->GetMetaData("FactorySourceProperty");
-			TSharedPtr<IPropertyHandle> FactorySourceHandle = PinnedPropertyHandle->GetParentHandle()->GetChildHandle(FName(*FactorySourcePropertyName));
-
-			if (FactorySourceHandle)
+			if (TSharedPtr<IPropertyHandle> PinnedPropertyHandle = WeakPropertyHandle.Pin())
 			{
-				UObject* FactorySourceObject = nullptr;
-				FAnimNextFactoryParams DefaultParams;
-				if (FactorySourceHandle->GetValue(FactorySourceObject) == FPropertyAccess::Result::Success && FactorySourceObject)
+				FString FactorySourcePropertyName = PinnedPropertyHandle->GetMetaData("FactorySourceProperty");
+				TSharedPtr<IPropertyHandle> FactorySourceHandle = PinnedPropertyHandle->GetParentHandle()->GetChildHandle(FName(*FactorySourcePropertyName));
+
+				if (FactorySourceHandle)
 				{
-					DefaultParams = UE::UAF::FAnimGraphFactory::GetDefaultParamsForObject(FactorySourceObject);
-				}
-				
-				PinnedPropertyHandle->EnumerateRawData([&DefaultParams](void* RawData, const int32 DataIndex, const int32 NumDatas)
+					UObject* FactorySourceObject = nullptr;
+					FAnimNextFactoryParams DefaultParams;
+					if (FactorySourceHandle->GetValue(FactorySourceObject) == FPropertyAccess::Result::Success && FactorySourceObject)
 					{
-						// todo: ideally we should retain settings on parameter structs that are of the same type
-						*static_cast<FAnimNextFactoryParams*>(RawData) = DefaultParams;
-						return true;
+						DefaultParams = UE::UAF::FAnimGraphFactory::GetDefaultParamsForObject(FactorySourceObject);
 					}
-				);
-			}
+				
+					PinnedPropertyHandle->EnumerateRawData([&DefaultParams](void* RawData, const int32 DataIndex, const int32 NumDatas)
+						{
+							// todo: ideally we should retain settings on parameter structs that are of the same type
+							*static_cast<FAnimNextFactoryParams*>(RawData) = DefaultParams;
+							return true;
+						}
+					);
+				}
 			
-			PinnedPropertyHandle->RequestRebuildChildren();
-		}
-	}));
-	
-	if (FactorySourceHandle.IsValid() && FactorySourceHandle->GetProperty()->IsA<FObjectProperty>())
-	{
-		FactorySourceHandle->EnumerateRawData([&CommonDefaults, &bIsInvalidOrAnimGraph](void* RawData, const int32 DataIndex, const int32 NumDatas)
+				PinnedPropertyHandle->RequestRebuildChildren();
+			}
+		}));
+
+		if (FactorySourceHandle->GetProperty()->IsA<FObjectProperty>())
 		{
-			UObject* Object = *(static_cast<UObject**>(RawData));
-		
-			if (!CommonDefaults.IsSet())
+			FactorySourceHandle->EnumerateRawData([&CommonDefaults, &bIsInvalidOrAnimGraph](void* RawData, const int32 DataIndex, const int32 NumDatas)
 			{
-				if (Object == nullptr || Object->IsA<UAnimNextAnimationGraph>())
+				UObject* Object = *(static_cast<UObject**>(RawData));
+			
+				if (!CommonDefaults.IsSet())
 				{
-					// TODO: At the moment we dont have 'details panel pin overrides' for assets, so we reset the params to not be confusing
-					// In the end we will want the params to also provide overrides like we supply to RunGraph nodes
-					CommonDefaults = FAnimNextFactoryParams();
-					bIsInvalidOrAnimGraph = true;
+					if (Object == nullptr || Object->IsA<UAnimNextAnimationGraph>())
+					{
+						// TODO: At the moment we dont have 'details panel pin overrides' for assets, so we reset the params to not be confusing
+						// In the end we will want the params to also provide overrides like we supply to RunGraph nodes
+						CommonDefaults = FAnimNextFactoryParams();
+						bIsInvalidOrAnimGraph = true;
+					}
+					else
+					{
+						CommonDefaults = FAnimGraphFactory::GetDefaultParamsForObject(Object);
+					}
 				}
 				else
 				{
-					CommonDefaults = FAnimGraphFactory::GetDefaultParamsForObject(Object);
-				}
-			}
-			else
-			{
-				if (Object)
-				{
-					const FAnimNextFactoryParams& DefaultParams = FAnimGraphFactory::GetDefaultParamsForObject(Object);
-					if (!Algo::Compare(DefaultParams.Builder.Stacks[0].TraitStructs, CommonDefaults->Builder.Stacks[0].TraitStructs))
+					if (Object)
 					{
-						CommonDefaults.Reset();
-						return false;
+						const FAnimNextFactoryParams& DefaultParams = FAnimGraphFactory::GetDefaultParamsForObject(Object);
+						if (!Algo::Compare(DefaultParams.Builder.Stacks[0].TraitStructs, CommonDefaults->Builder.Stacks[0].TraitStructs))
+						{
+							CommonDefaults.Reset();
+							return false;
+						}
 					}
 				}
-			}
-			return true;
-		});
+				return true;
+			});
+		}
 	}
 
 	TSet<const UScriptStruct*> AddedStructs;
