UAF: Initial Legacy Montage Support

diff --git a/Engine/Plugins/Experimental/UAF/UAFAnimGraph/Source/UAFAnimGraph/Internal/Injection/MontageTrait.h b/Engine/Plugins/Experimental/UAF/UAFAnimGraph/Source/UAFAnimGraph/Internal/Injection/MontageTrait.h
new file mode 100644
--- /dev/null
+++ b/Engine/Plugins/Experimental/UAF/UAFAnimGraph/Source/UAFAnimGraph/Internal/Injection/MontageTrait.h
@@ -0,0 +1,104 @@
+// Copyright Epic Games, Inc. All Rights Reserved.
+
+#pragma once
+
+#include "CoreMinimal.h"
+
+#include "TraitCore/Trait.h"
+#include "TraitInterfaces/IDiscreteBlend.h"
+#include "TraitInterfaces/IGarbageCollection.h"
+#include "TraitInterfaces/IUpdate.h"
+#include "Graph/AnimNextAnimationGraph.h"
+#include "Graph/AnimNextAnimGraph.h"
+#include "Graph/AnimNextGraphInstance.h"
+#include "TraitInterfaces/IBlendStack.h"
+#include "TraitInterfaces/IEvaluate.h"
+#include "TraitInterfaces/IHierarchy.h"
+#include "Animation/AnimSequence.h"
+
+#include "MontageTrait.generated.h"
+
+
+/** A trait that can play a montage based on the provided slot name */
+USTRUCT(meta = (DisplayName = "Montage Trait", ShowTooltip=true))
+struct FUAFMontageTraitSharedData : public FAnimNextTraitSharedData
+{
+	GENERATED_BODY()
+
+	UPROPERTY()
+	FAnimNextTraitHandle Source;
+
+	/** The name of this slot **/
+	UPROPERTY(EditAnywhere, Category = "Montage")
+	FName SlotName;
+
+	/** If true it will always update the source input, even when montages are fully blended in **/
+	UPROPERTY(EditAnywhere, Category = "Montage")
+	bool bAlwaysUpdateSource = false;
+
+	// Latent pin support boilerplate
+#define TRAIT_LATENT_PROPERTIES_ENUMERATOR(GeneratorMacro) \
+		GeneratorMacro(SlotName) \
+		GeneratorMacro(bAlwaysUpdateSource)
+
+GENERATE_TRAIT_LATENT_PROPERTIES(FUAFMontageTraitSharedData, TRAIT_LATENT_PROPERTIES_ENUMERATOR)
+#undef TRAIT_LATENT_PROPERTIES_ENUMERATOR
+
+};
+
+namespace UE::UAF
+{
+	struct FMontageTrait : FBaseTrait, IUpdate, IEvaluate, IHierarchy, IUpdateTraversal
+	{
+		DECLARE_ANIM_TRAIT(FMontageTrait, FBaseTrait)
+
+		struct FMontageEvalTaskData
+		{
+			float AnimPosition = 0.0f;
+			TWeakObjectPtr<UAnimSequence> WeakAnimationSequence = nullptr;
+
+			FMontageEvalTaskData(float InAnimPosition, UAnimSequence* InAnimSequence)
+				: AnimPosition(InAnimPosition)
+				, WeakAnimationSequence(InAnimSequence)
+			{ }
+		};
+
+		struct FInstanceData : FTrait::FInstanceData
+		{
+			/** Input node from which we receive the source pose to blend montages against during blend windows */
+			FTraitPtr Source;
+
+			/** The current blend weight of the source input */
+			float SourceWeight = 1.0f;
+
+			/** The source weight of the previous update, used to calculate if the source is blending out */
+			float LastUpdateSourceWeight = 0.0f;
+
+			/** Delta in seconds between updates, populated during PreUpdate */
+			float DeltaTime = 0.f;
+
+			/** The individual weights for each montage animation, plus source if relevant */
+			TArray<float> PoseWeights;
+
+			/** The anim data needed for the evaluation task for this slot */
+			TArray<FMontageEvalTaskData> EvalTaskDataForSlot;
+
+			void Construct(const FExecutionContext& Context, const FTraitBinding& Binding);
+		};
+
+		using FSharedData = FUAFMontageTraitSharedData;
+
+		// IUpdate impl
+		virtual void PreUpdate(FUpdateTraversalContext& Context, const TTraitBinding<IUpdate>& Binding, const FTraitUpdateState& TraitState) const override;
+
+		// IEvaluate impl
+		virtual void PostEvaluate(FEvaluateTraversalContext& Context, const TTraitBinding<IEvaluate>& Binding) const override;
+
+		// IHierarchy impl
+		virtual uint32 GetNumChildren(const FExecutionContext& Context, const TTraitBinding<IHierarchy>& Binding) const override;
+		virtual void GetChildren(const FExecutionContext& Context, const TTraitBinding<IHierarchy>& Binding, FChildrenArray& Children) const override;
+
+		// IUpdateTraversal impl
+		virtual void QueueChildrenForTraversal(FUpdateTraversalContext& Context, const TTraitBinding<IUpdateTraversal>& Binding, const FTraitUpdateState& TraitState, FUpdateTraversalQueue& TraversalQueue) const override;
+	};
+}
diff --git a/Engine/Plugins/Experimental/UAF/UAFAnimGraph/Source/UAFAnimGraph/Internal/Injection/UAFMontageComponent.h b/Engine/Plugins/Experimental/UAF/UAFAnimGraph/Source/UAFAnimGraph/Internal/Injection/UAFMontageComponent.h
new file mode 100644
--- /dev/null
+++ b/Engine/Plugins/Experimental/UAF/UAFAnimGraph/Source/UAFAnimGraph/Internal/Injection/UAFMontageComponent.h
@@ -0,0 +1,104 @@
+// Copyright Epic Games, Inc. All Rights Reserved.
+
+#pragma once
+
+#include "CoreMinimal.h"
+#include "Graph/UAFGraphInstanceComponent.h"
+#include "Module/UAFModuleInstanceComponent.h"
+#include "Engine/EngineBaseTypes.h"
+#include "AlphaBlend.h"
+#include "Animation/AnimNotifyQueue.h"
+#include "Animation/AnimMontage.h"
+#include "MontageTrait.h"
+#include "UAFMontageComponent.generated.h"
+
+struct FUAFUpdateMontageTickFunction : public FTickFunction
+{
+public:
+	// Begin FTickFunction overrides
+	virtual void ExecuteTick(float DeltaTime, enum ELevelTick TickType, ENamedThreads::Type CurrentThread, const FGraphEventRef& MyCompletionGraphEvent) override;
+	virtual FString DiagnosticMessage() override;
+	virtual FName DiagnosticContext(bool bDetailed) override;
+	// End FTickFunction overrides
+
+private:
+	TWeakObjectPtr<class UAnimNextComponent> Target = nullptr;
+	TWeakObjectPtr<class UAnimInstance> WeakAnimInstance = nullptr;
+	FUAFMontageComponent* MontageComponent = nullptr;
+
+	friend struct FUAFMontageComponent;
+};
+
+USTRUCT()
+struct FUAFMontageEvalData
+{
+	GENERATED_BODY()
+
+public:
+	FUAFMontageEvalData() = default;
+	FUAFMontageEvalData(UAnimMontage* InMontage, float InPosition, const FAlphaBlend& InBlendInfo)
+		: Montage(InMontage)
+		, BlendInfo(InBlendInfo)
+		, MontagePosition(InPosition)
+	{
+	}
+
+private:
+	UPROPERTY(Transient)
+	TObjectPtr<UAnimMontage> Montage = nullptr;
+
+	UPROPERTY(Transient)
+	FAlphaBlend BlendInfo = FAlphaBlend();
+
+	float MontagePosition = 0.0f;
+
+	friend struct UE::UAF::FMontageTrait;
+};
+
+// This module instance component is responsible for supporting montages within UAF. 
+// It handles required setup, update logic and provides relevant information to montage traits.
+USTRUCT()
+struct FUAFMontageComponent : public FUAFModuleInstanceComponent
+{
+	GENERATED_BODY()
+	DECLARE_UAF_ASSET_INSTANCE_COMPONENT()
+
+public:
+	FUAFMontageComponent() = default;
+	virtual ~FUAFMontageComponent() override;
+
+private:
+	virtual void OnBindToInstance() override;
+
+	// Retrieves all active montage animation notifies for a specific slot name
+	// returns true if any notifies for the given slot name exist
+	bool GetActiveNotifiesForSlot(FName SlotName, TArray<FAnimNotifyEventReference>& OutNotifies) const;
+
+private:
+	// All active and weighted montages, copied from the Anim Instance
+	UPROPERTY(Transient)
+	TArray<FUAFMontageEvalData> MontageEvalData;
+
+	// Active animation notifies per slot, used to have them easily accessible in montage traits 
+	UPROPERTY(Transient)
+	TMap<FName, FAnimNotifyArray> ActiveSlotNotifies;
+
+	// The tick function that is used to update montages on the game thread
+	FUAFUpdateMontageTickFunction MontageTickFunction;
+
+	// The event name for the tick function that the montage tick function will be enqueued as prerequisite for
+	// @Todo: Expose this to be able to be set by the user
+	static inline const FLazyName TickDependecyEventName = FLazyName("PrePhysics");
+
+	friend struct UE::UAF::FMontageTrait;
+	friend struct FUAFUpdateMontageTickFunction;
+};
+
+template<>
+struct TStructOpsTypeTraits<FUAFMontageComponent> : public TStructOpsTypeTraitsBase2<FUAFMontageComponent>
+{
+	enum
+	{
+		WithCopy = false,
+	};
+};
diff --git a/Engine/Plugins/Experimental/UAF/UAFAnimGraph/Source/UAFAnimGraph/Private/Injection/MontageTrait.cpp b/Engine/Plugins/Experimental/UAF/UAFAnimGraph/Source/UAFAnimGraph/Private/Injection/MontageTrait.cpp
new file mode 100644
--- /dev/null
+++ b/Engine/Plugins/Experimental/UAF/UAFAnimGraph/Source/UAFAnimGraph/Private/Injection/MontageTrait.cpp
@@ -0,0 +1,225 @@
+// Copyright Epic Games, Inc. All Rights Reserved.
+
+#include "Injection/MontageTrait.h"
+
+#include "InstanceTaskContext.h"
+#include "Factory/AnimGraphFactory.h"
+#include "TraitCore/ExecutionContext.h"
+#include "TraitCore/NodeInstance.h"
+#include "Graph/AnimNextGraphInstance.h"
+#include "Module/AnimNextModuleInstance.h"
+#include "EvaluationVM/Tasks/PushAnimSequenceKeyframe.h"
+#include "Animation/AnimMontage.h"
+#include "EvaluationVM/Tasks/BlendKeyframes.h"
+#include "EvaluationVM/Tasks/NormalizeRotations.h"
+#include "EvaluationVM/Tasks/PushReferenceKeyframe.h"
+#include "Traits/NotifyDispatcher.h"
+#include "Traits/NotifyFilterTrait.h"
+#include "Injection/UAFMontageComponent.h"
+
+namespace UE::UAF
+{
+	AUTO_REGISTER_ANIM_TRAIT(FMontageTrait)
+
+	// Trait implementation boilerplate
+	#define TRAIT_INTERFACE_ENUMERATOR(GeneratorMacro) \
+		GeneratorMacro(IUpdate) \
+		GeneratorMacro(IEvaluate) \
+		GeneratorMacro(IHierarchy) \
+		GeneratorMacro(IUpdateTraversal) \
+
+	GENERATE_ANIM_TRAIT_IMPLEMENTATION(FMontageTrait, TRAIT_INTERFACE_ENUMERATOR, NULL_ANIM_TRAIT_INTERFACE_ENUMERATOR, NULL_ANIM_TRAIT_EVENT_ENUMERATOR)
+	#undef TRAIT_INTERFACE_ENUMERATOR
+
+	void FMontageTrait::FInstanceData::Construct(const FExecutionContext& Context, const FTraitBinding& Binding)
+	{
+		FTrait::FInstanceData::Construct(Context, Binding);
+
+		if(FAnimNextModuleInstance* ModuleInstance = Context.GetRootGraphInstance().GetModuleInstance())
+		{
+			// Ensure the montage component is created at this point
+			ModuleInstance->GetComponent<FUAFMontageComponent>();
+		}
+
+		const FSharedData* SharedData = Binding.GetSharedData<FSharedData>();
+		FInstanceData* InstanceData = Binding.GetInstanceData<FInstanceData>();
+		if (!InstanceData->Source.IsValid())
+		{
+			InstanceData->Source = Context.AllocateNodeInstance(Binding, SharedData->Source);
+		}
+	}
+
+	void FMontageTrait::PreUpdate(FUpdateTraversalContext& Context, const TTraitBinding<IUpdate>& Binding, const FTraitUpdateState& TraitState) const
+	{
+		IUpdate::PreUpdate(Context, Binding, TraitState);
+
+		FInstanceData* InstanceData = Binding.GetInstanceData<FInstanceData>();
+		InstanceData->DeltaTime = TraitState.GetDeltaTime();
+		InstanceData->PoseWeights.Reset();
+		InstanceData->EvalTaskDataForSlot.Reset();
+		InstanceData->LastUpdateSourceWeight = InstanceData->SourceWeight;
+		InstanceData->SourceWeight = 1.0f;
+
+		FAnimNextModuleInstance* ModuleInstance = Context.GetRootGraphInstance().GetModuleInstance();
+		if (!ModuleInstance)
+		{
+			return;
+		}
+
+		const FUAFMontageComponent& MontageComponent = ModuleInstance->GetComponent<FUAFMontageComponent>();
+		if (MontageComponent.MontageEvalData.IsEmpty())
+		{
+			return;
+		}
+
+		const FSharedData* SharedData = Binding.GetSharedData<FSharedData>();
+		const FName CurrentSlotName = SharedData->GetSlotName(Binding);
+		if (CurrentSlotName.IsNone())
+		{
+			return;
+		}
+
+		float TotalMontageWeight = 0.0f;
+		for (const FUAFMontageEvalData& EvalData : MontageComponent.MontageEvalData)
+		{
+			if (EvalData.Montage != nullptr)
+			{
+				const UAnimMontage* const Montage = EvalData.Montage.Get();
+				if (Montage->IsValidSlot(CurrentSlotName))
+				{
+					if (const FAnimTrack* AnimTrack = EvalData.Montage->GetAnimationData(CurrentSlotName))
+					{
+						const float ClampedTime = FMath::Clamp(EvalData.MontagePosition, 0.f, AnimTrack->GetLength());
+						if (const FAnimSegment* AnimSegment = AnimTrack->GetSegmentAtTime(EvalData.MontagePosition))
+						{
+							float AnimPosition = 0;
+							if (UAnimSequence* AnimSegmentSequence = Cast<UAnimSequence>(AnimSegment->GetAnimationData(ClampedTime, AnimPosition)))
+							{
+								InstanceData->EvalTaskDataForSlot.Emplace(AnimPosition, AnimSegmentSequence);
+
+								const float MontageWeight = EvalData.BlendInfo.GetBlendedValue();
+								InstanceData->PoseWeights.Add(MontageWeight);
+								TotalMontageWeight += MontageWeight;
+							}
+						}
+					}
+				}
+			}
+		}
+
+		if (TotalMontageWeight > (1.f + ZERO_ANIMWEIGHT_THRESH))
+		{
+			for (int32 Index = 0; Index < InstanceData->PoseWeights.Num(); ++Index)
+			{
+				InstanceData->PoseWeights[Index] /= TotalMontageWeight;
+			}
+		}
+
+		const bool bAlwaysUpdateSource = SharedData->GetbAlwaysUpdateSource(Binding);
+		InstanceData->SourceWeight = 1.0f - FMath::Clamp(TotalMontageWeight, 0.0f, 1.0f);
+		if (FAnimWeight::IsRelevant(InstanceData->SourceWeight) || bAlwaysUpdateSource)
+		{
+			// Add the pose weight to blend in the source with
+			InstanceData->PoseWeights.Add(InstanceData->SourceWeight);
+		}
+
+		// Queue up anim notify events
+		if (FNotifyFilterTrait::AreNotifiesEnabledInScope(Context))
+		{
+			TArray<FAnimNotifyEventReference> NotifiesForSlot;
+			if (MontageComponent.GetActiveNotifiesForSlot(CurrentSlotName, NotifiesForSlot))
+			{
+				// Ensure we have a handler component on the module
+				ModuleInstance->GetComponent<FAnimNextNotifyDispatcherComponent>();
+
+				TSharedPtr<FNotifyDispatchEvent> NotifyDispatchEvent = MakeTraitEvent<FNotifyDispatchEvent>();
+				NotifyDispatchEvent->Notifies = MoveTemp(NotifiesForSlot);
+				NotifyDispatchEvent->Weight = 1.0f;
+				Context.RaiseOutputTraitEvent(NotifyDispatchEvent);
+			}
+		}
+	}
+
+	void FMontageTrait::PostEvaluate(FEvaluateTraversalContext& Context, const TTraitBinding<IEvaluate>& Binding) const
+	{
+		IEvaluate::PostEvaluate(Context, Binding);
+
+		FInstanceData* InstanceData = Binding.GetInstanceData<FInstanceData>();
+		const FSharedData* SharedData = Binding.GetSharedData<FSharedData>();
+		const bool bAlwaysUpdateSource = SharedData->GetbAlwaysUpdateSource(Binding);
+
+		// To be consistent with ABP behavior if there is nothing connected to the Source Input we push a ref pose onto the stack to blend against
+		// Another option would be to not blend if there isn't a source, however that would differ from ABP montage behavior
+		if (!InstanceData->Source.IsValid() && (FAnimWeight::IsRelevant(InstanceData->SourceWeight) || bAlwaysUpdateSource))
+		{
+			Context.AppendTask(FAnimNextPushReferenceKeyframeTask::MakeFromSkeleton());
+		}
+
+		if (InstanceData->EvalTaskDataForSlot.IsEmpty())
+		{
+			return;
+		}
+
+		const float DeltaTime = InstanceData->DeltaTime;
+		for (int32 iEvalTaskData = InstanceData->EvalTaskDataForSlot.Num() -1; iEvalTaskData >= 0; --iEvalTaskData)
+		{
+			const FMontageEvalTaskData& TaskData = InstanceData->EvalTaskDataForSlot[iEvalTaskData];
+			if (TaskData.WeakAnimationSequence.IsValid())
+			{
+				FAnimNextAnimSequenceKeyframeTask Task;
+				Task.bInterpolate = true;
+				Task.bLooping = false; 
+				Task.bExtractTrajectory = true;	/*Output.AnimInstanceProxy->ShouldExtractRootMotion()*/
+				Task.AnimSequence = TaskData.WeakAnimationSequence.Get();
+				Task.SampleTime = TaskData.AnimPosition;
+				Task.DeltaTimeRecord = FDeltaTimeRecord(DeltaTime);
+
+				Context.AppendTask(Task);
+			}
+		}
+
+		if (ensureMsgf(InstanceData->PoseWeights.Num() > 0, TEXT("Blending Montages: There are no pose weights to blend montages when we expected at least one.")))
+		{
+			Context.AppendTask(FAnimNextBlendOverwriteKeyframeWithScaleTask::Make(InstanceData->PoseWeights[0]));
+			for (int32 i = 1; i < InstanceData->PoseWeights.Num(); ++i)
+			{
+				Context.AppendTask(FAnimNextBlendAddKeyframeWithScaleTask::Make(InstanceData->PoseWeights[i]));
+			}
+
+			// Ensure that rotations are normalized after the blending 
+			Context.AppendTask(FAnimNextNormalizeKeyframeRotationsTask());
+		}
+	}
+
+	uint32 FMontageTrait::GetNumChildren(const FExecutionContext& Context, const TTraitBinding<IHierarchy>& Binding) const
+	{
+		return 1;
+	}
+
+	void FMontageTrait::GetChildren(const FExecutionContext& Context, const TTraitBinding<IHierarchy>& Binding, FChildrenArray& Children) const
+	{
+		const FInstanceData* InstanceData = Binding.GetInstanceData<FInstanceData>();
+
+		// Add the child, even if the handle is empty
+		Children.Add(InstanceData->Source);
+	}
+
+	void FMontageTrait::QueueChildrenForTraversal(FUpdateTraversalContext& Context, const TTraitBinding<IUpdateTraversal>& Binding, const FTraitUpdateState& TraitState, FUpdateTraversalQueue& TraversalQueue) const
+	{
+		FInstanceData* InstanceData = Binding.GetInstanceData<FInstanceData>();
+		const FSharedData* SharedData = Binding.GetSharedData<FSharedData>();
+
+		const bool bAlwaysUpdateSource = SharedData->GetbAlwaysUpdateSource(Binding);
+		if (InstanceData->Source.IsValid() && (FAnimWeight::IsRelevant(InstanceData->SourceWeight) || bAlwaysUpdateSource))
+		{
+			const bool bIsBlendingOut = InstanceData->LastUpdateSourceWeight > InstanceData->SourceWeight;
+			FTraitUpdateState TraitStateSource = TraitState
+				.WithWeight(InstanceData->SourceWeight)
+				.AsBlendingOut(bIsBlendingOut);
+
+			TraversalQueue.Push(InstanceData->Source, TraitStateSource);
+		}
+	}
+
+}
+
diff --git a/Engine/Plugins/Experimental/UAF/UAFAnimGraph/Source/UAFAnimGraph/Private/Injection/UAFMontageComponent.cpp b/Engine/Plugins/Experimental/UAF/UAFAnimGraph/Source/UAFAnimGraph/Private/Injection/UAFMontageComponent.cpp
new file mode 100644
--- /dev/null
+++ b/Engine/Plugins/Experimental/UAF/UAFAnimGraph/Source/UAFAnimGraph/Private/Injection/UAFMontageComponent.cpp
@@ -0,0 +1,190 @@
+// Copyright Epic Games, Inc. All Rights Reserved.
+
+#include "Injection/UAFMontageComponent.h"
+#include "Graph/AnimNextGraphInstance.h"
+#include "Module/AnimNextModuleInstance.h"
+#include "Module/AnimNextSkeletalMeshComponentReferenceComponent.h"
+#include "Component/AnimNextComponent.h"
+#include "Animation/AnimNotifies/AnimNotifyState.h"
+#include "Traits/NotifyDispatcher.h"
+#include "Animation/AnimInstance.h"
+#include "Animation/AnimMontage.h"
+#include "Components/SkeletalMeshComponent.h"
+#include "Engine/World.h"
+#include "Engine/Level.h"
+
+
+#include UE_INLINE_GENERATED_CPP_BY_NAME(UAFMontageComponent)
+
+void FUAFUpdateMontageTickFunction::ExecuteTick(float DeltaTime, enum ELevelTick TickType, ENamedThreads::Type CurrentThread, const FGraphEventRef& MyCompletionGraphEvent)
+{
+	if (!WeakAnimInstance.IsValid() || !MontageComponent)
+	{
+		return;
+	}
+
+	USkeletalMeshComponent* SkelMeshComponent = WeakAnimInstance->GetSkelMeshComponent();
+	if (!ensureMsgf(SkelMeshComponent != nullptr, TEXT("UAF Tick Montages: The SkelMeshComponent was unexpectedly null. Montages won't get ticked.")))
+	{
+		return;
+	}
+
+	// If the skelmesh component hasn't been ticked already (which can happen in CMC's PerformMovement if the animation has root motion) we make sure to tick it now to update montages 
+	SkelMeshComponent->TickMontagesOnly(DeltaTime, true);
+
+	FAnimNextModuleInstance* ModuleInstance = MontageComponent->GetModuleInstancePtr();
+	if (ensureMsgf(ModuleInstance != nullptr, TEXT("UAF TickMontages: Module Instance was not valid during tick, cannot send notify trait events!")))
+	{
+		// Queue up montage notifies, these will be fired no matter the slot weight. 
+		TSharedPtr<UE::UAF::FNotifyDispatchEvent> NotifyDispatchEvent = MakeTraitEvent<UE::UAF::FNotifyDispatchEvent>();
+		NotifyDispatchEvent->Notifies = WeakAnimInstance->NotifyQueue.AnimNotifies;
+		NotifyDispatchEvent->Weight = 1.0f;
+		ModuleInstance->QueueOutputTraitEvent(NotifyDispatchEvent);
+	}
+	
+	// Cache montage animation notifies, these will be fired only if the animation has weight. 
+	// Each slot will handle their notifies within the trait. 
+	// If you have more than one slot with the same name, these notifies might fire multiple times and may require manual filtering.
+	MontageComponent->ActiveSlotNotifies.Reset();
+	MontageComponent->ActiveSlotNotifies = WeakAnimInstance->NotifyQueue.UnfilteredMontageAnimNotifies;
+
+	MontageComponent->MontageEvalData.Reset();
+	for (FAnimMontageInstance* MontageInstance : WeakAnimInstance->MontageInstances)
+	{
+		// although montage can advance with 0.f weight, it is fine to filter by weight here
+		// because we don't want to evaluate them if 0 weight
+		if (MontageInstance->Montage && MontageInstance->GetWeight() > ZERO_ANIMWEIGHT_THRESH)
+		{
+			MontageComponent->MontageEvalData.Add(
+				FUAFMontageEvalData
+				(
+					MontageInstance->Montage,
+					MontageInstance->GetPosition(),
+					MontageInstance->GetBlend()
+
+				));
+		}
+	}
+}
+
+FString FUAFUpdateMontageTickFunction::DiagnosticMessage()
+{
+	return GetFullNameSafe(Target.Get()) + TEXT("[UAFUpdateMontageTickFunction]");
+}
+
+FName FUAFUpdateMontageTickFunction::DiagnosticContext(bool bDetailed)
+{
+	if (bDetailed)
+	{
+		return FName(*FString::Printf(TEXT("UAFUpdateMontageTickFunction/%s"), *GetFullNameSafe(Target.Get())));
+	}
+
+	static const FName UAFUpdateMontageTickFunctionFName(TEXT("UAFUpdateMontageTickFunction"));
+	return UAFUpdateMontageTickFunctionFName;
+}
+
+FUAFMontageComponent::~FUAFMontageComponent()
+{
+	// This component should only ever get destroyed from the game thread
+	// If this changes in the future we will need to revisit how to safely remove the tick prerequisite
+	check(IsInGameThread());
+
+	if (FAnimNextModuleInstance* ModuleInstance = GetModuleInstancePtr())
+	{
+		if (UAnimNextComponent* UAFComponent = Cast<UAnimNextComponent>(ModuleInstance->GetObject()))
+		{
+			UAFComponent->RemovePrerequisite(UAFComponent, MontageTickFunction, TickDependecyEventName);
+		}
+	}
+}
+
+bool FUAFMontageComponent::GetActiveNotifiesForSlot(FName SlotName, TArray<FAnimNotifyEventReference>& OutNotifies) const
+{
+	const FAnimNotifyArray* FoundEntry = ActiveSlotNotifies.Find(SlotName);
+	if (FoundEntry && FoundEntry->Notifies.Num() > 0)
+	{
+		OutNotifies = FoundEntry->Notifies;
+		return true;
+	}
+
+	return false;
+}
+
+void FUAFMontageComponent::OnBindToInstance()
+{
+	FAnimNextModuleInstance* ModuleInstance = GetModuleInstancePtr();
+	if (!ModuleInstance)
+	{
+		return;
+	}
+
+	const FAnimNextSkeletalMeshComponentReferenceComponent& ComponentReference = ModuleInstance->GetComponent<FAnimNextSkeletalMeshComponentReferenceComponent>();
+	USkeletalMeshComponent* SkeletalMeshComponent = ComponentReference.GetComponent();
+
+	if (!SkeletalMeshComponent || SkeletalMeshComponent->bEnableAnimation)
+	{
+		return;
+	}
+
+	UAnimNextComponent* UAFComponent = Cast<UAnimNextComponent>(ModuleInstance->GetObject());
+	if (!UAFComponent)
+	{
+		return;
+	}
+
+	MontageTickFunction.bRunOnAnyThread = false;
+	MontageTickFunction.bCanEverTick = true;
+	MontageTickFunction.bStartWithTickEnabled = true;
+	MontageTickFunction.MontageComponent = this;
+	MontageTickFunction.Target = UAFComponent;
+
+	FAnimNextModuleInstance::RunTaskOnGameThread(
+		[WeakSkelMeshComponent = TWeakObjectPtr<USkeletalMeshComponent>(SkeletalMeshComponent), WeakUAFComponent = TWeakObjectPtr<UAnimNextComponent>(UAFComponent), this]()
+		{
+
+			UAnimNextComponent* UAFComponent = WeakUAFComponent.Get();
+			if (!UAFComponent)
+			{
+				return;
+			}
+
+			AActor* UAFActor = UAFComponent->GetOwner();
+			if (!UAFActor)
+			{
+				return;
+			}
+
+			USkeletalMeshComponent* SkelMeshComponent = WeakSkelMeshComponent.Get();
+			if (!SkelMeshComponent)
+			{
+				return;
+			}
+
+			// Create a new anim instance if the skelmesh component doesn't have one already (which it should not if UAF is used) 
+			UAnimInstance* UsedAnimInstance = SkelMeshComponent->GetAnimInstance();
+			if (UsedAnimInstance == nullptr)
+			{
+				// Create a new AnimInstance that we can use to update our montages. 
+				SkelMeshComponent->AnimClass = UAnimInstance::StaticClass();
+				SkelMeshComponent->ClearAnimScriptInstance();
+				SkelMeshComponent->AnimScriptInstance = NewObject<UAnimInstance>(SkelMeshComponent, UAnimInstance::StaticClass());
+
+				UsedAnimInstance = SkelMeshComponent->GetAnimInstance();
+			}
+
+			// Do a lightweight init of the AnimInstance for montages only
+			UsedAnimInstance->InitializeMontageOnly();
+
+			// Setup and register tick function
+			MontageTickFunction.WeakAnimInstance = UsedAnimInstance;
+
+			if (UWorld* World = UAFComponent->GetWorld())
+			{
+				ULevel* Level = World->PersistentLevel;
+				MontageTickFunction.RegisterTickFunction(Level);
+
+				UAFComponent->AddPrerequisite(UAFComponent, MontageTickFunction, TickDependecyEventName);
+			}
+		});
+}
+
diff --git a/Engine/Source/Runtime/Engine/Classes/Animation/AnimInstance.h b/Engine/Source/Runtime/Engine/Classes/Animation/AnimInstance.h
--- a/Engine/Source/Runtime/Engine/Classes/Animation/AnimInstance.h
+++ b/Engine/Source/Runtime/Engine/Classes/Animation/AnimInstance.h
@@ -1340,6 +1340,8 @@ public:
 	// @param	bInDeferRootNodeInitialization	When set to true, defer init of the blend tree until the first Update() call
 	ENGINE_API void InitializeAnimation(bool bInDeferRootNodeInitialization = false);
 
+	ENGINE_API void InitializeMontageOnly();
+
 	/** Called on the game thread before UpdateAnimation is called on linked instances */
 	virtual void PreUpdateLinkedInstances(float DeltaSeconds) {}
 
@@ -1639,6 +1641,9 @@ private:
 	ENGINE_API void UpdateMontage(float DeltaSeconds);
 	ENGINE_API void UpdateMontageSyncGroup();
 
+public:
+	ENGINE_API void TickMontageOnly(float DeltaSeconds);
+
 protected:
 	// Updates the montage data used for evaluation based on the current playing montages
 	ENGINE_API void UpdateMontageEvaluationData();
diff --git a/Engine/Source/Runtime/Engine/Classes/Components/SkeletalMeshComponent.h b/Engine/Source/Runtime/Engine/Classes/Components/SkeletalMeshComponent.h
--- a/Engine/Source/Runtime/Engine/Classes/Components/SkeletalMeshComponent.h
+++ b/Engine/Source/Runtime/Engine/Classes/Components/SkeletalMeshComponent.h
@@ -1081,7 +1081,7 @@ public:
 	/**
 	 * Returns all tagged linked instance nodes that match the tag.
 	 */
-	UE_DEPRECATED(5.0, "Tags are unique so this funciton is no longer supported. Please use GetLinkedAnimGraphInstanceByTag instead")
+	UE_DEPRECATED(5.0, "Tags are unique so this function is no longer supported. Please use GetLinkedAnimGraphInstanceByTag instead")
 	UFUNCTION(BlueprintPure, Category = "Components|SkeletalMesh|Animation Blueprint Linking", meta = (Keywords = "AnimBlueprint", DeprecatedFunction, DeprecationMessage="Tags are unique so this function is no longer supported. Please use GetLinkedAnimGraphInstanceByTag instead"))
 	ENGINE_API void GetLinkedAnimGraphInstancesByTag(FName InTag, TArray<UAnimInstance*>& OutLinkedInstances) const;
 
@@ -1691,6 +1691,9 @@ public:
 	/** Tick Clothing Animation , basically this is called inside TickComponent */
 	ENGINE_API void TickClothing(float DeltaTime, FTickFunction& ThisTickFunction);
 
+	/** Tick only montages on the main anim instance to support legacy montages in UAF */
+	ENGINE_API void TickMontagesOnly(float DeltaTime, bool bNeedsValidRootMotion);
+
 	/** Store cloth simulation data into OutClothSimulData. Override USkinnedMeshComponent. */
 	ENGINE_API virtual void GetUpdateClothSimulationData_AnyThread(TMap<int32, FClothSimulData>& OutClothSimulData, FMatrix& OutLocalToWorld, float& OutClothBlendWeight) const override;
 
diff --git a/Engine/Source/Runtime/Engine/Private/Animation/AnimInstance.cpp b/Engine/Source/Runtime/Engine/Private/Animation/AnimInstance.cpp
--- a/Engine/Source/Runtime/Engine/Private/Animation/AnimInstance.cpp
+++ b/Engine/Source/Runtime/Engine/Private/Animation/AnimInstance.cpp
@@ -264,6 +264,22 @@ UWorld* UAnimInstance::GetWorld() const
 	return (HasAnyFlags(RF_ClassDefaultObject) ? nullptr : GetSkelMeshComponent()->GetWorld());
 }
 
+void UAnimInstance::InitializeMontageOnly()
+{
+	bUninitialized = false;
+	USkeletalMeshComponent* OwnerComponent = GetSkelMeshComponent();
+	if (OwnerComponent->GetSkeletalMeshAsset() != NULL)
+	{
+		CurrentSkeleton = OwnerComponent->GetSkeletalMeshAsset()->GetSkeleton();
+	}
+	else
+	{
+		CurrentSkeleton = NULL;
+	}
+
+	RecalcRequiredBones();
+}
+
 void UAnimInstance::InitializeAnimation(bool bInDeferRootNodeInitialization)
 {
 	FScopeCycleCounterUObject ContextScope(this);
@@ -511,6 +527,47 @@ void UAnimInstance::UpdateMontageSyncGroup()
 	}
 }
 
+void UAnimInstance::TickMontageOnly(float DeltaSeconds)
+{
+	ClearQueuedAnimEvents(true);
+
+	// Clear out any invalid montage instances 
+	for (int32 InstanceIndex = 0; InstanceIndex < MontageInstances.Num(); InstanceIndex++)
+	{
+		// Should never be null
+		FAnimMontageInstance* MontageInstance = MontageInstances[InstanceIndex];
+		ensure(MontageInstance);
+		if (MontageInstance && !MontageInstance->IsValid())
+		{
+			// Make sure we've cleared our references before deleting memory
+			ClearMontageInstanceReferences(*MontageInstance);
+
+			delete MontageInstance;
+			MontageInstances.RemoveAt(InstanceIndex);
+			--InstanceIndex;
+
+			if (MontageInstances.Num() == 0)
+			{
+				OnAllMontageInstancesEnded.Broadcast();
+			}
+		}
+	}
+
+	UpdateMontage(DeltaSeconds);
+	TriggerQueuedMontageEvents();
+	for (FAnimMontageInstance* MontageInstance : MontageInstances)
+	{
+		if (!MontageInstance)
+		{
+			continue;
+		}
+
+		MontageInstance->bDidUseMarkerSyncThisTick = false;
+		MontageInstance->MarkerTickRecord.Reset();
+		MontageInstance->MarkersPassedThisTick.Reset();
+	};
+}
+
 void UAnimInstance::UpdateAnimation(float DeltaSeconds, bool bNeedsValidRootMotion, EUpdateAnimationFlag UpdateFlag)
 {
 	LLM_SCOPE(ELLMTag::Animation);
diff --git a/Engine/Source/Runtime/Engine/Private/Components/CharacterMovementComponent.cpp b/Engine/Source/Runtime/Engine/Private/Components/CharacterMovementComponent.cpp
--- a/Engine/Source/Runtime/Engine/Private/Components/CharacterMovementComponent.cpp
+++ b/Engine/Source/Runtime/Engine/Private/Components/CharacterMovementComponent.cpp
@@ -11832,7 +11832,14 @@ void UCharacterMovementComponent::TickCharacterPose(float DeltaTime)
 		const bool bWasPlayingRootMotion = CharacterOwner->IsPlayingRootMotion();
 		const float RootMotionTranslationScale = CharacterOwner->GetAnimRootMotionTranslationScale();
 
-		CharacterMesh->TickPose(DeltaTime, true);
+		if (CharacterMesh->bEnableAnimation)
+		{
+			CharacterMesh->TickPose(DeltaTime, true);
+		}
+		else
+		{
+			CharacterMesh->TickMontagesOnly(DeltaTime, true);
+		}
 
 		// Grab root motion now that we have ticked the pose
 		if (CharacterOwner->IsPlayingRootMotion() || bWasPlayingRootMotion)
diff --git a/Engine/Source/Runtime/Engine/Private/Components/SkeletalMeshComponent.cpp b/Engine/Source/Runtime/Engine/Private/Components/SkeletalMeshComponent.cpp
--- a/Engine/Source/Runtime/Engine/Private/Components/SkeletalMeshComponent.cpp
+++ b/Engine/Source/Runtime/Engine/Private/Components/SkeletalMeshComponent.cpp
@@ -1247,7 +1247,7 @@ bool USkeletalMeshComponent::InitializeAnimScriptInstance(bool bForceReinit, boo
 						InputNode->CachedInputPose.SetBoneContainer(&BoneContainer);
 					}
 
-					// SetBoneContainer allocates space for bone data but leaves it uninitalized.
+					// SetBoneContainer allocates space for bone data but leaves it uninitialized.
 					InputNode->bIsCachedInputPoseInitialized = false;
 				}
 
@@ -1649,6 +1649,33 @@ void USkeletalMeshComponent::TickAnimInstances(float DeltaTime, bool bNeedsValid
 	}
 }
 
+void USkeletalMeshComponent::TickMontagesOnly(float DeltaTime, bool bNeedsValidRootMotion)
+{
+	if (!bEnableAnimation)
+	{
+		if (AnimScriptInstance)
+		{
+			if (!ShouldTickPose())
+			{
+				return;
+			}
+
+			// Update the pose tick frame counter to avoid double ticking the montages. 
+			LastPoseTickFrame = static_cast<uint32>(GFrameCounter);
+
+			AnimScriptInstance->TickMontageOnly(DeltaTime);
+		}
+		else
+		{
+			UE_LOG(LogAnimation, Warning, TEXT("SkeletalMeshComponent in '%s' tried to TickMontagesOnly but there is no valid AnimInstance"), *GetNameSafe(GetOwner()));
+		}
+	}
+	else
+	{
+		UE_LOG(LogAnimation, Warning, TEXT("SkeletalMeshComponent in '%s' tried to call TickMontagesOnly when animations are enabled. This is likely an error and will be skipped."), *GetNameSafe(GetOwner()));
+	}
+}
+
 bool USkeletalMeshComponent::UpdateLODStatus()
 {
 	if (Super::UpdateLODStatus())
@@ -5057,7 +5084,7 @@ bool USkeletalMeshComponent::MoveComponentImpl(const FVector& Delta, const FQuat
 	bool bSuccess = Super::MoveComponentImpl(Delta, NewRotation, bSweep, OutHit, MoveFlags, Teleport);
 	if(bSuccess && Teleport != ETeleportType::None)
 	{
-		// If a skeletal mesh component recieves a teleport we should reset any other dynamic simulations
+		// If a skeletal mesh component receives a teleport we should reset any other dynamic simulations
 		ResetAnimInstanceDynamics(Teleport);
 
 		OnSkelMeshPhysicsTeleported.Broadcast();
