UAF - Add nodiscard to GetComponent function and add explicitly to GetComponent calls that happen only to ensure the component is created at that point

diff --git a/Engine/Plugins/Experimental/UAF/UAF/Source/UAF/Public/UAFAssetInstance.h b/Engine/Plugins/Experimental/UAF/UAF/Source/UAF/Public/UAFAssetInstance.h
--- a/Engine/Plugins/Experimental/UAF/UAF/Source/UAF/Public/UAFAssetInstance.h
+++ b/Engine/Plugins/Experimental/UAF/UAF/Source/UAF/Public/UAFAssetInstance.h
@@ -146,7 +146,7 @@ struct FUAFAssetInstance
 	// Returns a typed instance component, creating it lazily the first time it is queried.
 	// When querying for structs of a base type, existing derived types will be returned, rather than a new base struct created.
 	template<typename ComponentType>
-	ComponentType& GetComponent()
+	[[nodiscard]] ComponentType& GetComponent()
 	{
 		static_assert(std::is_base_of<FUAFAssetInstanceComponent, ComponentType>::value, "ComponentType type must derive from FUAFAssetInstanceComponent");
 
diff --git a/Engine/Plugins/Experimental/UAF/UAFAnimGraph/Source/UAFAnimGraph/Private/Injection/InjectionSiteTrait.cpp b/Engine/Plugins/Experimental/UAF/UAFAnimGraph/Source/UAFAnimGraph/Private/Injection/InjectionSiteTrait.cpp
--- a/Engine/Plugins/Experimental/UAF/UAFAnimGraph/Source/UAFAnimGraph/Private/Injection/InjectionSiteTrait.cpp
+++ b/Engine/Plugins/Experimental/UAF/UAFAnimGraph/Source/UAFAnimGraph/Private/Injection/InjectionSiteTrait.cpp
@@ -63,7 +63,7 @@ namespace UE::UAF
 
 		if(FAnimNextModuleInstance* ModuleInstance = Context.GetRootGraphInstance().GetModuleInstance())
 		{
-			ModuleInstance->GetComponent<FAnimNextModuleInjectionComponent>();
+			(void)ModuleInstance->GetComponent<FAnimNextModuleInjectionComponent>();
 		}
 
 		IGarbageCollection::RegisterWithGC(Context, Binding);
diff --git a/Engine/Plugins/Experimental/UAF/UAFAnimGraph/Source/UAFAnimGraph/Private/Injection/MontageTrait.cpp b/Engine/Plugins/Experimental/UAF/UAFAnimGraph/Source/UAFAnimGraph/Private/Injection/MontageTrait.cpp
--- a/Engine/Plugins/Experimental/UAF/UAFAnimGraph/Source/UAFAnimGraph/Private/Injection/MontageTrait.cpp
+++ b/Engine/Plugins/Experimental/UAF/UAFAnimGraph/Source/UAFAnimGraph/Private/Injection/MontageTrait.cpp
@@ -38,7 +38,7 @@ namespace UE::UAF
 		if(FAnimNextModuleInstance* ModuleInstance = Context.GetRootGraphInstance().GetModuleInstance())
 		{
 			// Ensure the montage component is created at this point
-			ModuleInstance->GetComponent<FUAFMontageComponent>();
+			(void)ModuleInstance->GetComponent<FUAFMontageComponent>();
 		}
 
 		const FSharedData* SharedData = Binding.GetSharedData<FSharedData>();
@@ -130,7 +130,7 @@ namespace UE::UAF
 			if (MontageComponent.GetActiveNotifiesForSlot(CurrentSlotName, NotifiesForSlot))
 			{
 				// Ensure we have a handler component on the module
-				ModuleInstance->GetComponent<FAnimNextNotifyDispatcherComponent>();
+				(void)ModuleInstance->GetComponent<FAnimNextNotifyDispatcherComponent>();
 
 				TSharedPtr<FNotifyDispatchEvent> NotifyDispatchEvent = MakeTraitEvent<FNotifyDispatchEvent>();
 				NotifyDispatchEvent->Notifies = MoveTemp(NotifiesForSlot);
diff --git a/Engine/Plugins/Experimental/UAF/UAFAnimGraph/Source/UAFAnimGraph/Private/Traits/NotifyDispatcher.cpp b/Engine/Plugins/Experimental/UAF/UAFAnimGraph/Source/UAFAnimGraph/Private/Traits/NotifyDispatcher.cpp
--- a/Engine/Plugins/Experimental/UAF/UAFAnimGraph/Source/UAFAnimGraph/Private/Traits/NotifyDispatcher.cpp
+++ b/Engine/Plugins/Experimental/UAF/UAFAnimGraph/Source/UAFAnimGraph/Private/Traits/NotifyDispatcher.cpp
@@ -70,7 +70,7 @@ namespace UE::UAF
 			// Ensure we have a handler component on the module
 			if(FAnimNextModuleInstance* ModuleInstance = Context.GetRootGraphInstance().GetModuleInstance())
 			{
-				ModuleInstance->GetComponent<FAnimNextNotifyDispatcherComponent>();
+				(void)ModuleInstance->GetComponent<FAnimNextNotifyDispatcherComponent>();
 			}
 
 			auto NotifyDispatchEvent = MakeTraitEvent<FNotifyDispatchEvent>();
