UAF - Added group based sync marker synchronization support for legacy montages

diff --git a/Engine/Plugins/Experimental/UAF/UAFAnimGraph/Source/UAFAnimGraph/Internal/Injection/MontageTrait.h b/Engine/Plugins/Experimental/UAF/UAFAnimGraph/Source/UAFAnimGraph/Internal/Injection/MontageTrait.h
--- a/Engine/Plugins/Experimental/UAF/UAFAnimGraph/Source/UAFAnimGraph/Internal/Injection/MontageTrait.h
+++ b/Engine/Plugins/Experimental/UAF/UAFAnimGraph/Source/UAFAnimGraph/Internal/Injection/MontageTrait.h
@@ -15,9 +15,12 @@
 #include "TraitInterfaces/IEvaluate.h"
 #include "TraitInterfaces/IHierarchy.h"
 #include "Animation/AnimSequence.h"
+#include "TraitInterfaces/ITimeline.h"
+#include "TraitInterfaces/IGroupSynchronization.h"
 
 #include "MontageTrait.generated.h"
 
+struct FUAFMontageComponent;
 
 /** A trait that can play a montage based on the provided slot name */
 USTRUCT(meta = (DisplayName = "Montage Trait", ShowTooltip=true))
@@ -28,18 +31,23 @@ struct FUAFMontageTraitSharedData : public FAnimNextTraitSharedData
 	UPROPERTY()
 	FAnimNextTraitHandle Source;
 
-	/** The name of this slot **/
+	/** The name of this slot */
 	UPROPERTY(EditAnywhere, Category = "Montage")
 	FName SlotName;
 
-	/** If true it will always update the source input, even when montages are fully blended in **/
+	/** If true it will always update the source input, even when montages are fully blended in */
 	UPROPERTY(EditAnywhere, Category = "Montage")
 	bool bAlwaysUpdateSource = false;
 
+	/** Disables synchronization for this slot even if the playing montage has a sync group setup */
+	UPROPERTY(EditAnywhere, Category = "Montage")
+	bool bDisableSynchronization = false;
+
 	// Latent pin support boilerplate
 #define TRAIT_LATENT_PROPERTIES_ENUMERATOR(GeneratorMacro) \
 		GeneratorMacro(SlotName) \
-		GeneratorMacro(bAlwaysUpdateSource)
+		GeneratorMacro(bAlwaysUpdateSource) \
+		GeneratorMacro(bDisableSynchronization)
 
 GENERATE_TRAIT_LATENT_PROPERTIES(FUAFMontageTraitSharedData, TRAIT_LATENT_PROPERTIES_ENUMERATOR)
 #undef TRAIT_LATENT_PROPERTIES_ENUMERATOR
@@ -48,7 +56,7 @@ GENERATE_TRAIT_LATENT_PROPERTIES(FUAFMontageTraitSharedData, TRAIT_LATENT_PROPER
 
 namespace UE::UAF
 {
-	struct FMontageTrait : FBaseTrait, IUpdate, IEvaluate, IHierarchy, IUpdateTraversal
+	struct FMontageTrait : FBaseTrait, IUpdate, IEvaluate, IHierarchy, IUpdateTraversal, ITimeline, IGroupSynchronization
 	{
 		DECLARE_ANIM_TRAIT(FMontageTrait, FBaseTrait)
 
@@ -83,6 +91,24 @@ namespace UE::UAF
 			/** The anim data needed for the evaluation task for this slot */
 			TArray<FMontageEvalTaskData> EvalTaskDataForSlot;
 
+			/** The state representing the timeline before the montage update */
+			FTimelineState PreMontageUpdateState = FTimelineState();
+
+			/** The state representing the timeline after the montage update */
+			FTimelineState PostMontageUpdateState = FTimelineState();
+
+			/** If the pre-update or post-update state should be returned on GetState() */
+			bool bReturnPreUpdateState = true;
+
+			/** The sync group params used to participate in synchronization */
+			FSyncGroupParameters SyncGroupParams;
+
+			/** The eval data index into the montage component for the active synchronization candidate */
+			int32 SyncCandidateIndex = INDEX_NONE;
+
+			/** Cached Ptr to the montage component */
+			FUAFMontageComponent* MontageComponent = nullptr;
+
 			void Construct(const FExecutionContext& Context, const FTraitBinding& Binding);
 		};
 
@@ -100,5 +126,13 @@ namespace UE::UAF
 
 		// IUpdateTraversal impl
 		virtual void QueueChildrenForTraversal(FUpdateTraversalContext& Context, const TTraitBinding<IUpdateTraversal>& Binding, const FTraitUpdateState& TraitState, FUpdateTraversalQueue& TraversalQueue) const override;
+
+		// ITimeline impl
+		virtual void GetSyncMarkers(const FExecutionContext& Context, const TTraitBinding<ITimeline>& Binding, FTimelineSyncMarkerArray& OutSyncMarkers) const override;
+		virtual FTimelineState GetState(const FExecutionContext& Context, const TTraitBinding<ITimeline>& Binding) const override;
+
+		// IGroupSynchronization impl
+		virtual FSyncGroupParameters GetGroupParameters(FExecutionContext& Context, const TTraitBinding<IGroupSynchronization>& Binding) const override;
+		virtual void AdvanceBy(FExecutionContext& Context, const TTraitBinding<IGroupSynchronization>& Binding, float DeltaTime, bool bDispatchEvents) const override;
 	};
 }
diff --git a/Engine/Plugins/Experimental/UAF/UAFAnimGraph/Source/UAFAnimGraph/Internal/Injection/UAFMontageComponent.h b/Engine/Plugins/Experimental/UAF/UAFAnimGraph/Source/UAFAnimGraph/Internal/Injection/UAFMontageComponent.h
--- a/Engine/Plugins/Experimental/UAF/UAFAnimGraph/Source/UAFAnimGraph/Internal/Injection/UAFMontageComponent.h
+++ b/Engine/Plugins/Experimental/UAF/UAFAnimGraph/Source/UAFAnimGraph/Internal/Injection/UAFMontageComponent.h
@@ -36,10 +36,11 @@ struct FUAFMontageEvalData
 
 public:
 	FUAFMontageEvalData() = default;
-	FUAFMontageEvalData(UAnimMontage* InMontage, float InPosition, const FAlphaBlend& InBlendInfo)
+	FUAFMontageEvalData(UAnimMontage* InMontage, float InPosition, const FAlphaBlend& InBlendInfo, float InPlayRate)
 		: Montage(InMontage)
 		, BlendInfo(InBlendInfo)
 		, MontagePosition(InPosition)
+		, MontagePlayRate(InPlayRate)
 	{
 	}
 
@@ -51,6 +52,7 @@ private:
 	FAlphaBlend BlendInfo = FAlphaBlend();
 
 	float MontagePosition = 0.0f;
+	float MontagePlayRate = 1.0f;
 
 	friend struct UE::UAF::FMontageTrait;
 };
diff --git a/Engine/Plugins/Experimental/UAF/UAFAnimGraph/Source/UAFAnimGraph/Private/Injection/MontageTrait.cpp b/Engine/Plugins/Experimental/UAF/UAFAnimGraph/Source/UAFAnimGraph/Private/Injection/MontageTrait.cpp
--- a/Engine/Plugins/Experimental/UAF/UAFAnimGraph/Source/UAFAnimGraph/Private/Injection/MontageTrait.cpp
+++ b/Engine/Plugins/Experimental/UAF/UAFAnimGraph/Source/UAFAnimGraph/Private/Injection/MontageTrait.cpp
@@ -16,6 +16,7 @@
 #include "Traits/NotifyDispatcher.h"
 #include "Traits/NotifyFilterTrait.h"
 #include "Injection/UAFMontageComponent.h"
+#include "Graph/SyncGroup_GraphInstanceComponent.h"
 
 namespace UE::UAF
 {
@@ -27,6 +28,8 @@ namespace UE::UAF
 		GeneratorMacro(IEvaluate) \
 		GeneratorMacro(IHierarchy) \
 		GeneratorMacro(IUpdateTraversal) \
+		GeneratorMacro(ITimeline) \
+		GeneratorMacro(IGroupSynchronization) \
 
 	GENERATE_ANIM_TRAIT_IMPLEMENTATION(FMontageTrait, TRAIT_INTERFACE_ENUMERATOR, NULL_ANIM_TRAIT_INTERFACE_ENUMERATOR, NULL_ANIM_TRAIT_EVENT_ENUMERATOR)
 	#undef TRAIT_INTERFACE_ENUMERATOR
@@ -37,8 +40,8 @@ namespace UE::UAF
 
 		if(FAnimNextModuleInstance* ModuleInstance = Context.GetRootGraphInstance().GetModuleInstance())
 		{
-			// Ensure the montage component is created at this point
-			(void)ModuleInstance->GetComponent<FUAFMontageComponent>();
+			// Cache the montage component and create it lazily if it hasn't been created yet. 
+			MontageComponent = &ModuleInstance->GetComponent<FUAFMontageComponent>();
 		}
 
 		const FSharedData* SharedData = Binding.GetSharedData<FSharedData>();
@@ -47,6 +50,12 @@ namespace UE::UAF
 		{
 			InstanceData->Source = Context.AllocateNodeInstance(Binding, SharedData->Source);
 		}
+
+		InstanceData->SyncGroupParams.bMatchSyncPoint = false;
+		InstanceData->SyncGroupParams.SyncMode = EAnimGroupSynchronizationMode::SynchronizeUsingGroupName;
+		// Montage will always be the leader because they cannot sync to any other animation within UAF as they are advanced outside of UAF
+		// Instead we solely use montages to drive synchronization of other assets in the UAF system
+		InstanceData->SyncGroupParams.GroupRole = EAnimGroupSynchronizationRole::AlwaysLeader;
 	}
 
 	void FMontageTrait::PreUpdate(FUpdateTraversalContext& Context, const TTraitBinding<IUpdate>& Binding, const FTraitUpdateState& TraitState) const
@@ -59,15 +68,16 @@ namespace UE::UAF
 		InstanceData->EvalTaskDataForSlot.Reset();
 		InstanceData->LastUpdateSourceWeight = InstanceData->SourceWeight;
 		InstanceData->SourceWeight = 1.0f;
+		InstanceData->bReturnPreUpdateState = true;
+		InstanceData->SyncCandidateIndex = INDEX_NONE;
 
-		FAnimNextModuleInstance* ModuleInstance = Context.GetRootGraphInstance().GetModuleInstance();
-		if (!ModuleInstance)
+		const FUAFMontageComponent* MontageComponent = InstanceData->MontageComponent;
+		if (MontageComponent == nullptr)
 		{
 			return;
 		}
 
-		const FUAFMontageComponent& MontageComponent = ModuleInstance->GetComponent<FUAFMontageComponent>();
-		if (MontageComponent.MontageEvalData.IsEmpty())
+		if (MontageComponent->MontageEvalData.IsEmpty())
 		{
 			return;
 		}
@@ -79,9 +89,13 @@ namespace UE::UAF
 			return;
 		}
 
+		const bool bIsSyncDisabled = SharedData->GetbDisableSynchronization(Binding);
 		float TotalMontageWeight = 0.0f;
-		for (const FUAFMontageEvalData& EvalData : MontageComponent.MontageEvalData)
+		float SyncCandidate_HeighestWeight = -1.0f;
+		int32 SyncCandidate_HeighestWeightIndex = INDEX_NONE;
+		for (int32 i = 0; i < MontageComponent->MontageEvalData.Num(); ++i)
 		{
+			const FUAFMontageEvalData& EvalData = MontageComponent->MontageEvalData[i];
 			if (EvalData.Montage != nullptr)
 			{
 				const UAnimMontage* const Montage = EvalData.Montage.Get();
@@ -100,6 +114,16 @@ namespace UE::UAF
 								const float MontageWeight = EvalData.BlendInfo.GetBlendedValue();
 								InstanceData->PoseWeights.Add(MontageWeight);
 								TotalMontageWeight += MontageWeight;
+
+								if (!bIsSyncDisabled)
+								{
+									const bool bIsSynchronizationCandidate = !EvalData.Montage->SyncGroup.IsNone() && EvalData.BlendInfo.GetDesiredValue() == 0 && EvalData.BlendInfo.IsComplete() == false;
+									if (bIsSynchronizationCandidate && MontageWeight > SyncCandidate_HeighestWeight)
+									{
+										SyncCandidate_HeighestWeight = MontageWeight;
+										SyncCandidate_HeighestWeightIndex = i;
+									}
+								}
 							}
 						}
 					}
@@ -127,10 +151,13 @@ namespace UE::UAF
 		if (FNotifyFilterTrait::AreNotifiesEnabledInScope(Context))
 		{
 			TArray<FAnimNotifyEventReference> NotifiesForSlot;
-			if (MontageComponent.GetActiveNotifiesForSlot(CurrentSlotName, NotifiesForSlot))
+			if (MontageComponent->GetActiveNotifiesForSlot(CurrentSlotName, NotifiesForSlot))
 			{
-				// Ensure we have a handler component on the module
-				(void)ModuleInstance->GetComponent<FAnimNextNotifyDispatcherComponent>();
+				if (FAnimNextModuleInstance* ModuleInstance = Context.GetRootGraphInstance().GetModuleInstance())
+				{
+					// Ensure we have a handler component on the module
+					(void)ModuleInstance->GetComponent<FAnimNextNotifyDispatcherComponent>();
+				}
 
 				TSharedPtr<FNotifyDispatchEvent> NotifyDispatchEvent = MakeTraitEvent<FNotifyDispatchEvent>();
 				NotifyDispatchEvent->Notifies = MoveTemp(NotifiesForSlot);
@@ -138,6 +165,33 @@ namespace UE::UAF
 				Context.RaiseOutputTraitEvent(NotifyDispatchEvent);
 			}
 		}
+
+		// Handle Group Synchronization data and registering
+		if (MontageComponent->MontageEvalData.IsValidIndex(SyncCandidate_HeighestWeightIndex))
+		{
+			const FUAFMontageEvalData& SyncCandidateEvalData = MontageComponent->MontageEvalData[SyncCandidate_HeighestWeightIndex];
+			const float MontageLength = SyncCandidateEvalData.Montage->CalculateSequenceLength();
+
+			// On the first activation of synchronization the timeline state will be invalid, so we create it manually once 
+			if (InstanceData->PostMontageUpdateState.GetDuration() == 0.0f)
+			{
+				const float PreUpdatePosition = FMath::Clamp(SyncCandidateEvalData.MontagePosition - InstanceData->DeltaTime, 0.0f, MontageLength);
+				InstanceData->PreMontageUpdateState = FTimelineState(PreUpdatePosition, MontageLength, SyncCandidateEvalData.MontagePlayRate, false).WithDebugName(SyncCandidateEvalData.Montage->GetFName());
+			}
+			else
+			{
+				InstanceData->PreMontageUpdateState = InstanceData->PostMontageUpdateState;
+			}
+
+			InstanceData->PostMontageUpdateState = FTimelineState(SyncCandidateEvalData.MontagePosition, MontageLength, SyncCandidateEvalData.MontagePlayRate, false).WithDebugName(SyncCandidateEvalData.Montage->GetFName());
+			InstanceData->SyncGroupParams.GroupName = SyncCandidateEvalData.Montage->SyncGroup;
+			InstanceData->SyncCandidateIndex = SyncCandidate_HeighestWeightIndex;
+
+			// Register with synchronization component
+			const FSyncGroupParameters GroupParameters = InstanceData->SyncGroupParams;
+			FSyncGroupGraphInstanceComponent& SyncComponent = Context.GetComponent<FSyncGroupGraphInstanceComponent>();
+			SyncComponent.RegisterWithGroup(GroupParameters, Binding.GetTraitPtr(), TraitState);
+		}
 	}
 
 	void FMontageTrait::PostEvaluate(FEvaluateTraversalContext& Context, const TTraitBinding<IEvaluate>& Binding) const
@@ -221,5 +275,50 @@ namespace UE::UAF
 		}
 	}
 
+	void FMontageTrait::GetSyncMarkers(const FExecutionContext& Context, const TTraitBinding<ITimeline>& Binding, FTimelineSyncMarkerArray& OutSyncMarkers) const
+	{
+		const FInstanceData* InstanceData = Binding.GetInstanceData<FInstanceData>();
+		if (InstanceData->SyncCandidateIndex == INDEX_NONE || InstanceData->MontageComponent == nullptr)
+		{
+			return;
+		}
+
+		if (InstanceData->MontageComponent->MontageEvalData.IsValidIndex(InstanceData->SyncCandidateIndex))
+		{
+			const FUAFMontageEvalData& SyncCandidateEvalData = InstanceData->MontageComponent->MontageEvalData[InstanceData->SyncCandidateIndex];
+			if (ensureMsgf(SyncCandidateEvalData.Montage, TEXT("Get Montage SyncMarkers: SyncCandidateIndex referred to eval data with an invalid montage!")))
+			{
+				const int32 NumSyncMarkers = SyncCandidateEvalData.Montage->MarkerData.AuthoredSyncMarkers.Num();
+				if (NumSyncMarkers > 0)
+				{
+					OutSyncMarkers.Reserve(NumSyncMarkers);
+					for (const FAnimSyncMarker& SyncMarker : SyncCandidateEvalData.Montage->MarkerData.AuthoredSyncMarkers)
+					{
+						OutSyncMarkers.Add(FTimelineSyncMarker(SyncMarker.MarkerName, SyncMarker.Time));
+					}
+				}
+			}
+		}
+	}
+
+	FTimelineState FMontageTrait::GetState(const FExecutionContext& Context, const TTraitBinding<ITimeline>& Binding) const
+	{
+		const FInstanceData* InstanceData = Binding.GetInstanceData<FInstanceData>();
+		return InstanceData->bReturnPreUpdateState ? InstanceData->PreMontageUpdateState : InstanceData->PostMontageUpdateState;
+	}
+
+	void FMontageTrait::AdvanceBy(FExecutionContext& Context, const TTraitBinding<IGroupSynchronization>& Binding, float DeltaTime, bool bDispatchEvents) const
+	{
+		// This function gets called by the SyncGroupComponent and would usually advance the timeline of this trait
+		// Because montages get currently advanced outside of UAF we mimic this behaviour by swapping which state we return in GetState()
+		FInstanceData* InstanceData = Binding.GetInstanceData<FInstanceData>();
+		InstanceData->bReturnPreUpdateState = false;
+	}
+
+	FSyncGroupParameters FMontageTrait::GetGroupParameters(FExecutionContext& Context, const TTraitBinding<IGroupSynchronization>& Binding) const
+	{
+		const FInstanceData* InstanceData = Binding.GetInstanceData<FInstanceData>();
+		return InstanceData->SyncGroupParams;
+	}
 }
 
diff --git a/Engine/Plugins/Experimental/UAF/UAFAnimGraph/Source/UAFAnimGraph/Private/Injection/UAFMontageComponent.cpp b/Engine/Plugins/Experimental/UAF/UAFAnimGraph/Source/UAFAnimGraph/Private/Injection/UAFMontageComponent.cpp
--- a/Engine/Plugins/Experimental/UAF/UAFAnimGraph/Source/UAFAnimGraph/Private/Injection/UAFMontageComponent.cpp
+++ b/Engine/Plugins/Experimental/UAF/UAFAnimGraph/Source/UAFAnimGraph/Private/Injection/UAFMontageComponent.cpp
@@ -60,8 +60,8 @@ void FUAFUpdateMontageTickFunction::ExecuteTick(float DeltaTime, enum ELevelTick
 				(
 					MontageInstance->Montage,
 					MontageInstance->GetPosition(),
-					MontageInstance->GetBlend()
-
+					MontageInstance->GetBlend(),
+					MontageInstance->GetPlayRate()
 				));
 		}
 	}
