Fix false positive compilation errors for IGarbageCollection
* Update error message to display the pin in question
* Update tests to include testing for this

#rnx

diff --git a/Engine/Plugins/Experimental/UAF/UAFAnimGraph/Source/UAFAnimGraphUncookedOnly/Private/Graph/AnimNextAnimationGraph_EditorData.cpp b/Engine/Plugins/Experimental/UAF/UAFAnimGraph/Source/UAFAnimGraphUncookedOnly/Private/Graph/AnimNextAnimationGraph_EditorData.cpp
--- a/Engine/Plugins/Experimental/UAF/UAFAnimGraph/Source/UAFAnimGraphUncookedOnly/Private/Graph/AnimNextAnimationGraph_EditorData.cpp
+++ b/Engine/Plugins/Experimental/UAF/UAFAnimGraph/Source/UAFAnimGraphUncookedOnly/Private/Graph/AnimNextAnimationGraph_EditorData.cpp
@@ -434,17 +434,17 @@ namespace UE::UAF::UncookedOnly::Private
 					{
 						// This pin has something linked to it, it is a latent pin
 						FString PinCPPType = Pin->GetCPPType();
-						if (RigVMTypeUtils::ObjectFromCPPType(PinCPPType))
+						if (RigVMTypeUtils::IsUObjectType(PinCPPType))
 						{
 							const bool bImplementsGarbageCollection = TraitEntry.Trait->GetTraitInterface(IGarbageCollection::InterfaceUID) != nullptr;
 							if (!bImplementsGarbageCollection)
 							{
 								const FText MissingInterfaceError = FText::Format(
-									LOCTEXT("TraitStatusInStack_DoesNotImplementGarbageCollection", "@@: Trait is connected to a UObject but does not implement IGarbageCollection ({0})"),
-									TraitEntry.Trait->GetTraitSharedDataStruct()->GetDisplayNameText());
+									LOCTEXT("TraitStatusInStack_DoesNotImplementGarbageCollection", "@@: Trait is connected to a UObject but does not implement IGarbageCollection ({0} - {1})"),
+									TraitEntry.Trait->GetTraitSharedDataStruct()->GetDisplayNameText(),
+									FText::FromName(Pin->GetDisplayName()));
 								InSettings.ASTSettings.Report(EMessageSeverity::Error, InVMUnitNode, MissingInterfaceError.ToString());
 								bResult = false;
-								break;
 							}
 						}
 					}
diff --git a/Engine/Plugins/Experimental/UAF/UAFTests/Source/UAFAnimGraphTestSuite/Private/GCTestsUtil.h b/Engine/Plugins/Experimental/UAF/UAFTests/Source/UAFAnimGraphTestSuite/Private/GCTestsUtil.h
--- a/Engine/Plugins/Experimental/UAF/UAFTests/Source/UAFAnimGraphTestSuite/Private/GCTestsUtil.h
+++ b/Engine/Plugins/Experimental/UAF/UAFTests/Source/UAFAnimGraphTestSuite/Private/GCTestsUtil.h
@@ -5,6 +5,7 @@
 #include "TraitCore/TraitBinding.h"
 #include "TraitCore/TraitSharedData.h"
 #include "Animation/AnimSequence.h"
+#include "Animation/AnimTypes.h"
 #include "Graph/AnimNextGraphInstance.h"
 #include "Graph/RigUnit_AnimNextBase.h"
 
@@ -17,9 +18,15 @@ struct FUAFTestAnimSequenceSharedData : public FAnimNextTraitSharedData
 
 	UPROPERTY(EditAnywhere, Category = "Default")
 	TObjectPtr<UAnimSequence> AnimSequence;
+	
+	// Adding a UEnum here to test for false positives when validating trait stacks.
+	// This pin shouldn't ever require GC handling so let's test it *doesn't* generate an error.
+	UPROPERTY(EditAnywhere, Category = "Default")
+	EAnimInterpolationType DummyProp0 = EAnimInterpolationType::Linear;
 
 	#define TRAIT_LATENT_PROPERTIES_ENUMERATOR(GeneratorMacro) \
 		GeneratorMacro(AnimSequence) \
+		GeneratorMacro(DummyProp0) \
 
 	GENERATE_TRAIT_LATENT_PROPERTIES(FUAFTestAnimSequenceSharedData, TRAIT_LATENT_PROPERTIES_ENUMERATOR)
 	#undef TRAIT_LATENT_PROPERTIES_ENUMERATOR
