UAF - Add Additive Montage support

diff --git a/Engine/Plugins/Experimental/UAF/UAFAnimGraph/Source/UAFAnimGraph/Internal/Injection/MontageTrait.h b/Engine/Plugins/Experimental/UAF/UAFAnimGraph/Source/UAFAnimGraph/Internal/Injection/MontageTrait.h
--- a/Engine/Plugins/Experimental/UAF/UAFAnimGraph/Source/UAFAnimGraph/Internal/Injection/MontageTrait.h
+++ b/Engine/Plugins/Experimental/UAF/UAFAnimGraph/Source/UAFAnimGraph/Internal/Injection/MontageTrait.h
@@ -88,9 +88,15 @@ namespace UE::UAF
 			/** The individual weights for each montage animation, plus source if relevant */
 			TArray<float> PoseWeights;
 
+			/** The individual weights for each additive montage animation */
+			TArray<float> AdditivePoseWeights;
+
 			/** The anim data needed for the evaluation task for this slot */
 			TArray<FMontageEvalTaskData> EvalTaskDataForSlot;
 
+			/** The additive anim data needed for the evaluation task for this slot */
+			TArray<FMontageEvalTaskData> AdditiveEvalTaskDataForSlot;
+
 			/** The state representing the timeline before the montage update */
 			FTimelineState PreMontageUpdateState = FTimelineState();
 
diff --git a/Engine/Plugins/Experimental/UAF/UAFAnimGraph/Source/UAFAnimGraph/Private/Injection/MontageTrait.cpp b/Engine/Plugins/Experimental/UAF/UAFAnimGraph/Source/UAFAnimGraph/Private/Injection/MontageTrait.cpp
--- a/Engine/Plugins/Experimental/UAF/UAFAnimGraph/Source/UAFAnimGraph/Private/Injection/MontageTrait.cpp
+++ b/Engine/Plugins/Experimental/UAF/UAFAnimGraph/Source/UAFAnimGraph/Private/Injection/MontageTrait.cpp
@@ -17,6 +17,7 @@
 #include "Traits/NotifyFilterTrait.h"
 #include "Injection/UAFMontageComponent.h"
 #include "Graph/SyncGroup_GraphInstanceComponent.h"
+#include "EvaluationVM/Tasks/ApplyAdditiveKeyframe.h"
 
 namespace UE::UAF
 {
@@ -66,6 +67,8 @@ namespace UE::UAF
 		InstanceData->DeltaTime = TraitState.GetDeltaTime();
 		InstanceData->PoseWeights.Reset();
 		InstanceData->EvalTaskDataForSlot.Reset();
+		InstanceData->AdditiveEvalTaskDataForSlot.Reset();
+		InstanceData->AdditivePoseWeights.Reset();
 		InstanceData->LastUpdateSourceWeight = InstanceData->SourceWeight;
 		InstanceData->SourceWeight = 1.0f;
 		InstanceData->bReturnPreUpdateState = true;
@@ -91,6 +94,7 @@ namespace UE::UAF
 
 		const bool bIsSyncDisabled = SharedData->GetbDisableSynchronization(Binding);
 		float TotalMontageWeight = 0.0f;
+		float TotalNonAdditiveMontageWeight = 0.0f;
 		float SyncCandidate_HeighestWeight = -1.0f;
 		int32 SyncCandidate_HeighestWeightIndex = INDEX_NONE;
 		for (int32 i = 0; i < MontageComponent->MontageEvalData.Num(); ++i)
@@ -109,11 +113,19 @@ namespace UE::UAF
 							float AnimPosition = 0;
 							if (UAnimSequence* AnimSegmentSequence = Cast<UAnimSequence>(AnimSegment->GetAnimationData(ClampedTime, AnimPosition)))
 							{
-								InstanceData->EvalTaskDataForSlot.Emplace(AnimPosition, AnimSegmentSequence);
-
 								const float MontageWeight = EvalData.BlendInfo.GetBlendedValue();
-								InstanceData->PoseWeights.Add(MontageWeight);
 								TotalMontageWeight += MontageWeight;
+								if (AnimTrack->IsAdditive())
+								{
+									InstanceData->AdditivePoseWeights.Add(MontageWeight);
+									InstanceData->AdditiveEvalTaskDataForSlot.Emplace(AnimPosition, AnimSegmentSequence);
+								}
+								else
+								{
+									InstanceData->PoseWeights.Add(MontageWeight);
+									InstanceData->EvalTaskDataForSlot.Emplace(AnimPosition, AnimSegmentSequence);
+									TotalNonAdditiveMontageWeight += MontageWeight;
+								}
 
 								if (!bIsSyncDisabled)
 								{
@@ -133,15 +145,26 @@ namespace UE::UAF
 
 		if (TotalMontageWeight > (1.f + ZERO_ANIMWEIGHT_THRESH))
 		{
+			// If we have to re-normalize the weights, we also have to recalculate the total non-additive weight 
+			TotalNonAdditiveMontageWeight = 0.0f;
 			for (int32 Index = 0; Index < InstanceData->PoseWeights.Num(); ++Index)
 			{
 				InstanceData->PoseWeights[Index] /= TotalMontageWeight;
+				TotalNonAdditiveMontageWeight += InstanceData->PoseWeights[Index];
 			}
+
+			for (int32 Index = 0; Index < InstanceData->AdditivePoseWeights.Num(); ++Index)
+			{
+				InstanceData->AdditivePoseWeights[Index] /= TotalMontageWeight;
+			}
+
 		}
 
+		// If we are only playing additive animations we do not add the source weight as we want to apply 100% of it and not blend it in
+		const bool bOnlyPlayingAdditives = InstanceData->EvalTaskDataForSlot.Num() == 0 && InstanceData->AdditiveEvalTaskDataForSlot.Num() > 0;
 		const bool bAlwaysUpdateSource = SharedData->GetbAlwaysUpdateSource(Binding);
-		InstanceData->SourceWeight = 1.0f - FMath::Clamp(TotalMontageWeight, 0.0f, 1.0f);
-		if (FAnimWeight::IsRelevant(InstanceData->SourceWeight) || bAlwaysUpdateSource)
+		InstanceData->SourceWeight = 1.0f - FMath::Clamp(TotalNonAdditiveMontageWeight, 0.0f, 1.0f);
+		if ((FAnimWeight::IsRelevant(InstanceData->SourceWeight) || bAlwaysUpdateSource) && !bOnlyPlayingAdditives)
 		{
 			// Add the pose weight to blend in the source with
 			InstanceData->PoseWeights.Add(InstanceData->SourceWeight);
@@ -201,19 +224,21 @@ namespace UE::UAF
 		FInstanceData* InstanceData = Binding.GetInstanceData<FInstanceData>();
 		const FSharedData* SharedData = Binding.GetSharedData<FSharedData>();
 		const bool bAlwaysUpdateSource = SharedData->GetbAlwaysUpdateSource(Binding);
+		const bool bOnlyPlayingAdditives = InstanceData->EvalTaskDataForSlot.Num() == 0 && InstanceData->AdditiveEvalTaskDataForSlot.Num() > 0;
 
 		// To be consistent with ABP behavior if there is nothing connected to the Source Input we push a ref pose onto the stack to blend against
 		// Another option would be to not blend if there isn't a source, however that would differ from ABP montage behavior
-		if (!InstanceData->Source.IsValid() && (FAnimWeight::IsRelevant(InstanceData->SourceWeight) || bAlwaysUpdateSource))
+		if (!InstanceData->Source.IsValid() && (FAnimWeight::IsRelevant(InstanceData->SourceWeight) || bAlwaysUpdateSource || bOnlyPlayingAdditives))
 		{
 			Context.AppendTask(FAnimNextPushReferenceKeyframeTask::MakeFromSkeleton());
 		}
 
-		if (InstanceData->EvalTaskDataForSlot.IsEmpty())
+		if (InstanceData->EvalTaskDataForSlot.IsEmpty() && InstanceData->AdditiveEvalTaskDataForSlot.IsEmpty())
 		{
 			return;
 		}
 
+		// First we push all non-additive animation keyframes onto the stack
 		const float DeltaTime = InstanceData->DeltaTime;
 		for (int32 iEvalTaskData = InstanceData->EvalTaskDataForSlot.Num() -1; iEvalTaskData >= 0; --iEvalTaskData)
 		{
@@ -232,7 +257,8 @@ namespace UE::UAF
 			}
 		}
 
-		if (ensureMsgf(InstanceData->PoseWeights.Num() > 0, TEXT("Blending Montages: There are no pose weights to blend montages when we expected at least one.")))
+		// Then we blend the non-additive poses and source pose (if applicable) together 
+		if (InstanceData->PoseWeights.Num() > 0)
 		{
 			Context.AppendTask(FAnimNextBlendOverwriteKeyframeWithScaleTask::Make(InstanceData->PoseWeights[0]));
 			for (int32 i = 1; i < InstanceData->PoseWeights.Num(); ++i)
@@ -243,6 +269,31 @@ namespace UE::UAF
 			// Ensure that rotations are normalized after the blending 
 			Context.AppendTask(FAnimNextNormalizeKeyframeRotationsTask());
 		}
+		
+		// Finally we add the additive poses and blend them in
+		if (InstanceData->AdditivePoseWeights.Num() > 0)
+		{
+			check(InstanceData->AdditivePoseWeights.Num() == InstanceData->AdditiveEvalTaskDataForSlot.Num());
+
+			for (int32 iAdditivePose = 0; iAdditivePose < InstanceData->AdditiveEvalTaskDataForSlot.Num(); iAdditivePose++)
+			{
+				const FMontageEvalTaskData& TaskData = InstanceData->AdditiveEvalTaskDataForSlot[iAdditivePose];
+				if (TaskData.WeakAnimationSequence.IsValid())
+				{
+					FAnimNextAnimSequenceKeyframeTask Task;
+					Task.bInterpolate = true;
+					Task.bLooping = false;
+					Task.bExtractTrajectory = true;	/*Output.AnimInstanceProxy->ShouldExtractRootMotion()*/
+					Task.AnimSequence = TaskData.WeakAnimationSequence.Get();
+					Task.SampleTime = TaskData.AnimPosition;
+					Task.DeltaTimeRecord = FDeltaTimeRecord(DeltaTime);
+
+					Context.AppendTask(Task);
+
+					Context.AppendTask(FAnimNextApplyAdditiveKeyframeTask::Make(InstanceData->AdditivePoseWeights[iAdditivePose]));
+				}
+			}
+		}
 	}
 
 	uint32 FMontageTrait::GetNumChildren(const FExecutionContext& Context, const TTraitBinding<IHierarchy>& Binding) const
@@ -264,7 +315,8 @@ namespace UE::UAF
 		const FSharedData* SharedData = Binding.GetSharedData<FSharedData>();
 
 		const bool bAlwaysUpdateSource = SharedData->GetbAlwaysUpdateSource(Binding);
-		if (InstanceData->Source.IsValid() && (FAnimWeight::IsRelevant(InstanceData->SourceWeight) || bAlwaysUpdateSource))
+		const bool bOnlyPlayingAdditives = InstanceData->EvalTaskDataForSlot.Num() == 0 && InstanceData->AdditiveEvalTaskDataForSlot.Num() > 0;
+		if (InstanceData->Source.IsValid() && (FAnimWeight::IsRelevant(InstanceData->SourceWeight) || bAlwaysUpdateSource || bOnlyPlayingAdditives))
 		{
 			const bool bIsBlendingOut = InstanceData->LastUpdateSourceWeight > InstanceData->SourceWeight;
 			FTraitUpdateState TraitStateSource = TraitState
