Add nullptr checks to avoid crashing when the new abstract skeleton runtime is used in UAF

diff --git a/Engine/Plugins/Experimental/UAF/UAFAnimGraph/Source/UAFAnimGraph/Private/EvaluationVM/EvaluationVM.cpp b/Engine/Plugins/Experimental/UAF/UAFAnimGraph/Source/UAFAnimGraph/Private/EvaluationVM/EvaluationVM.cpp
--- a/Engine/Plugins/Experimental/UAF/UAFAnimGraph/Source/UAFAnimGraph/Private/EvaluationVM/EvaluationVM.cpp
+++ b/Engine/Plugins/Experimental/UAF/UAFAnimGraph/Source/UAFAnimGraph/Private/EvaluationVM/EvaluationVM.cpp
@@ -96,7 +96,7 @@ namespace UE::UAF
 
 		FKeyframeState Keyframe;
 
-		if (EnumHasAnyFlags(EvaluationFlags, EEvaluationFlags::Bones))
+		if (EnumHasAnyFlags(EvaluationFlags, EEvaluationFlags::Bones) && ReferencePose != nullptr)
 		{
 			const bool bInitWithRefPose = true;
 
@@ -115,7 +115,7 @@ namespace UE::UAF
 	{
 		FKeyframeState Keyframe;
 
-		if (EnumHasAnyFlags(EvaluationFlags, EEvaluationFlags::Bones))
+		if (EnumHasAnyFlags(EvaluationFlags, EEvaluationFlags::Bones) && ReferencePose != nullptr)
 		{
 			const bool bInitWithRefPose = false;
 
