Fix comments

diff --git a/Engine/Plugins/Experimental/UAF/UAF/Source/UAF/Public/UAF/Attributes/Transformers/Overwrite.h b/Engine/Plugins/Experimental/UAF/UAF/Source/UAF/Public/UAF/Attributes/Transformers/Overwrite.h
--- a/Engine/Plugins/Experimental/UAF/UAF/Source/UAF/Public/UAF/Attributes/Transformers/Overwrite.h
+++ b/Engine/Plugins/Experimental/UAF/UAF/Source/UAF/Public/UAF/Attributes/Transformers/Overwrite.h
@@ -12,7 +12,7 @@
 namespace UE::UAF::Transformers
 {
 	// Attribute value transformer to overwrite values
-	// Accumulate is the initial step within an N-step interpolation: Overwrite, Accumulate0, Accumulate1, ...
+	// Overwrite is the initial step within an N-step interpolation: Overwrite, Accumulate0, Accumulate1, ...
 	// Make sure to also implement variations for Accumulate and Interpolate
 	struct FOverwrite final : public FAttributeValueTransformer
 	{
diff --git a/Engine/Plugins/Experimental/UAF/UAFAnimGraph/Source/UAFAnimGraph/Private/EvaluationVM/Tasks/BlendKeyframes.cpp b/Engine/Plugins/Experimental/UAF/UAFAnimGraph/Source/UAFAnimGraph/Private/EvaluationVM/Tasks/BlendKeyframes.cpp
--- a/Engine/Plugins/Experimental/UAF/UAFAnimGraph/Source/UAFAnimGraph/Private/EvaluationVM/Tasks/BlendKeyframes.cpp
+++ b/Engine/Plugins/Experimental/UAF/UAFAnimGraph/Source/UAFAnimGraph/Private/EvaluationVM/Tasks/BlendKeyframes.cpp
@@ -218,8 +218,6 @@ void FAnimNextBlendOverwriteKeyframeWithScaleTask::Execute(UE::UAF::FEvaluationV
 
 	if (VM.GetContext().IsValid())
 	{
-		// Pop our top two poses, we'll re-use the top keyframe for our result
-
 		TUniquePtr<FAttributeCollection> Collection;
 		if (!VM.PopValue(ATTRIBUTE_STACK_NAME, Collection))
 		{
diff --git a/Engine/Plugins/Experimental/UAF/UAFAnimGraph/Source/UAFAnimGraph/Public/EvaluationVM/EvaluationVM.h b/Engine/Plugins/Experimental/UAF/UAFAnimGraph/Source/UAFAnimGraph/Public/EvaluationVM/EvaluationVM.h
--- a/Engine/Plugins/Experimental/UAF/UAFAnimGraph/Source/UAFAnimGraph/Public/EvaluationVM/EvaluationVM.h
+++ b/Engine/Plugins/Experimental/UAF/UAFAnimGraph/Source/UAFAnimGraph/Public/EvaluationVM/EvaluationVM.h
@@ -331,7 +331,7 @@ namespace UE::UAF
 	extern UAFANIMGRAPH_API const FEvaluationVMStackName KEYFRAME_STACK_NAME;
 	ANIM_NEXT_ENABLE_EVALUATION_STACK_USAGE(TUniquePtr<FKeyframeState>)
 
-	// A stack of FAttributeCollectionStack instances used when sampling sequences and blending their results
+	// A stack of FAttributeCollection instances used when sampling sequences and blending their results
 	extern UAFANIMGRAPH_API const FEvaluationVMStackName ATTRIBUTE_STACK_NAME;
 	ANIM_NEXT_ENABLE_EVALUATION_STACK_USAGE(TUniquePtr<FAttributeCollection>)
 }
