Fix invalid UAF trait pointer when dynamic trait storage is re-used

diff --git a/Engine/Plugins/Experimental/UAF/UAFAnimGraph/Source/UAFAnimGraph/Internal/TraitCore/TraitRegistry.h b/Engine/Plugins/Experimental/UAF/UAFAnimGraph/Source/UAFAnimGraph/Internal/TraitCore/TraitRegistry.h
--- a/Engine/Plugins/Experimental/UAF/UAFAnimGraph/Source/UAFAnimGraph/Internal/TraitCore/TraitRegistry.h
+++ b/Engine/Plugins/Experimental/UAF/UAFAnimGraph/Source/UAFAnimGraph/Internal/TraitCore/TraitRegistry.h
@@ -106,7 +106,7 @@ namespace UE::UAF
 		uint8							StaticTraitBuffer[STATIC_TRAIT_BUFFER_SIZE] = { 0 };
 		int32							StaticTraitBufferOffset = 0;
 
-		TArray<uintptr_t>				DynamicTraits;
+		TArray<intptr_t>				DynamicTraits;
 		int32							DynamicTraitFreeIndexHead = INDEX_NONE;
 
 		TMap<FTraitUIDRaw, FRegistryEntry>	TraitUIDToEntryMap;
diff --git a/Engine/Plugins/Experimental/UAF/UAFAnimGraph/Source/UAFAnimGraph/Private/TraitCore/TraitRegistry.cpp b/Engine/Plugins/Experimental/UAF/UAFAnimGraph/Source/UAFAnimGraph/Private/TraitCore/TraitRegistry.cpp
--- a/Engine/Plugins/Experimental/UAF/UAFAnimGraph/Source/UAFAnimGraph/Private/TraitCore/TraitRegistry.cpp
+++ b/Engine/Plugins/Experimental/UAF/UAFAnimGraph/Source/UAFAnimGraph/Private/TraitCore/TraitRegistry.cpp
@@ -150,12 +150,17 @@ namespace UE::UAF
 				{
 					// We already had a free index, grab it
 					TraitIndex = DynamicTraitFreeIndexHead;
+
+					// Update the next free index
 					DynamicTraitFreeIndexHead = DynamicTraits[DynamicTraitFreeIndexHead];
+
+					// Set our trait ptr to our new index
+					DynamicTraits[TraitIndex] = reinterpret_cast<intptr_t>(Trait);
 				}
 				else
 				{
 					// No free indices, allocate a new one
-					TraitIndex = DynamicTraits.Add(reinterpret_cast<uintptr_t>(Trait));
+					TraitIndex = DynamicTraits.Add(reinterpret_cast<intptr_t>(Trait));
 				}
 
 				TraitHandle = FTraitRegistryHandle::MakeDynamic(TraitIndex);
@@ -326,14 +331,17 @@ namespace UE::UAF
 			{
 				// We already had a free index, grab it
 				TraitIndex = DynamicTraitFreeIndexHead;
+
+				// Update the next free index
 				DynamicTraitFreeIndexHead = DynamicTraits[DynamicTraitFreeIndexHead];
 
-				DynamicTraits[TraitIndex] = reinterpret_cast<uintptr_t>(Trait);
+				// Set our trait ptr to our new index
+				DynamicTraits[TraitIndex] = reinterpret_cast<intptr_t>(Trait);
 			}
 			else
 			{
 				// No free indices, allocate a new one
-				TraitIndex = DynamicTraits.Add(reinterpret_cast<uintptr_t>(Trait));
+				TraitIndex = DynamicTraits.Add(reinterpret_cast<intptr_t>(Trait));
 			}
 
 			FTraitRegistryHandle TraitHandle = FTraitRegistryHandle::MakeDynamic(TraitIndex);
