UAF: Add option to FAnimNextBlendOverwriteKeyframePerBoneWithScaleTask and FAnimNextBlendAddKeyframePerBoneWithScaleTask to blend poses together based on per bone blend weights without a blendprofile

diff --git a/Engine/Plugins/Experimental/UAF/UAFAnimGraph/Source/UAFAnimGraph/Private/EvaluationVM/Tasks/BlendKeyframesPerBone.cpp b/Engine/Plugins/Experimental/UAF/UAFAnimGraph/Source/UAFAnimGraph/Private/EvaluationVM/Tasks/BlendKeyframesPerBone.cpp
--- a/Engine/Plugins/Experimental/UAF/UAFAnimGraph/Source/UAFAnimGraph/Private/EvaluationVM/Tasks/BlendKeyframesPerBone.cpp
+++ b/Engine/Plugins/Experimental/UAF/UAFAnimGraph/Source/UAFAnimGraph/Private/EvaluationVM/Tasks/BlendKeyframesPerBone.cpp
@@ -18,13 +18,23 @@ FAnimNextBlendOverwriteKeyframePerBoneWithScaleTask FAnimNextBlendOverwriteKeyfr
 {
 	FAnimNextBlendOverwriteKeyframePerBoneWithScaleTask Task;
 	Task.BlendProfile = BlendProfile;
-	Task.SourceSkeleton = BlendProfile->GetSkeleton();
 	Task.BlendData = &BlendData;
 	Task.ScaleFactor = ScaleFactor;
 
 	return Task;
 }
 
+FAnimNextBlendOverwriteKeyframePerBoneWithScaleTask FAnimNextBlendOverwriteKeyframePerBoneWithScaleTask::Make(const FBlendSampleData& BlendData, float ScaleFactor, bool bPerBoneDataUsesSkeletonIndex)
+{
+	FAnimNextBlendOverwriteKeyframePerBoneWithScaleTask Task;
+	Task.BlendProfile = nullptr;
+	Task.BlendData = &BlendData;
+	Task.ScaleFactor = ScaleFactor;
+	Task.bPerBoneDataUsesSkeletonIndex = bPerBoneDataUsesSkeletonIndex;
+
+	return Task;
+}
+
 void FAnimNextBlendOverwriteKeyframePerBoneWithScaleTask::Execute(UE::UAF::FEvaluationVM& VM) const
 {
 	using namespace UE::UAF;
@@ -36,9 +46,9 @@ void FAnimNextBlendOverwriteKeyframePerBoneWithScaleTask::Execute(UE::UAF::FEval
 		return;
 	}
 
-	if (BlendProfile == nullptr)
+	if (BlendProfile == nullptr && BlendData->PerBoneBlendData.Num() == 0)
 	{
-		// If we don't have a blend profile, we blend the whole pose
+		// If we don't have a blend profile and per bone blend data, we blend the whole pose
 		Super::Execute(VM);
 		return;
 	}
@@ -50,35 +60,56 @@ void FAnimNextBlendOverwriteKeyframePerBoneWithScaleTask::Execute(UE::UAF::FEval
 		return;
 	}
 
-	const USkeleton* TargetSkeleton = Keyframe->Pose.GetSkeletonAsset();
-	TSharedPtr<IInterpolationIndexProvider::FPerBoneInterpolationData> Data = BlendProfile->GetPerBoneInterpolationData(TargetSkeleton);
-
-	const FSkeletonRemapping& SkeletonRemapping = UE::Anim::FSkeletonRemappingRegistry::Get().GetRemapping(BlendProfile->GetSkeleton(), TargetSkeleton);
-
 	const TArrayView<const FBoneIndexType> LODBoneIndexToSkeletonBoneIndexMap = Keyframe->Pose.GetLODBoneIndexToSkeletonBoneIndexMap();
 	const int32 NumLODBones = LODBoneIndexToSkeletonBoneIndexMap.Num();
 
 	TArray<int32> LODBoneIndexToWeightIndexMap;
 	LODBoneIndexToWeightIndexMap.AddUninitialized(NumLODBones);
 
-	if (SkeletonRemapping.IsValid())
+	if (BlendProfile)
 	{
-		for (int32 LODBoneIndex = 0; LODBoneIndex < NumLODBones; ++LODBoneIndex)
+		const USkeleton* TargetSkeleton = Keyframe->Pose.GetSkeletonAsset();
+		TSharedPtr<IInterpolationIndexProvider::FPerBoneInterpolationData> Data = BlendProfile->GetPerBoneInterpolationData(TargetSkeleton);
+
+		const FSkeletonRemapping& SkeletonRemapping = UE::Anim::FSkeletonRemappingRegistry::Get().GetRemapping(BlendProfile->GetSkeleton(), TargetSkeleton);
+
+		if (SkeletonRemapping.IsValid())
 		{
-			const FSkeletonPoseBoneIndex TargetSkeletonBoneIndex(LODBoneIndexToSkeletonBoneIndexMap[LODBoneIndex]);
-			const FSkeletonPoseBoneIndex SourceSkeletonBoneIndex(SkeletonRemapping.GetSourceSkeletonBoneIndex(TargetSkeletonBoneIndex.GetInt()));
-			LODBoneIndexToWeightIndexMap[LODBoneIndex] = SourceSkeletonBoneIndex.IsValid() ? BlendProfile->GetPerBoneInterpolationIndex(SourceSkeletonBoneIndex, TargetSkeleton, Data.Get()) : INDEX_NONE;
+			for (int32 LODBoneIndex = 0; LODBoneIndex < NumLODBones; ++LODBoneIndex)
+			{
+				const FSkeletonPoseBoneIndex TargetSkeletonBoneIndex(LODBoneIndexToSkeletonBoneIndexMap[LODBoneIndex]);
+				const FSkeletonPoseBoneIndex SourceSkeletonBoneIndex(SkeletonRemapping.GetSourceSkeletonBoneIndex(TargetSkeletonBoneIndex.GetInt()));
+				LODBoneIndexToWeightIndexMap[LODBoneIndex] = SourceSkeletonBoneIndex.IsValid() ? BlendProfile->GetPerBoneInterpolationIndex(SourceSkeletonBoneIndex, TargetSkeleton, Data.Get()) : INDEX_NONE;
+			}
+		}
+		else
+		{
+			for (int32 LODBoneIndex = 0; LODBoneIndex < NumLODBones; ++LODBoneIndex)
+			{
+				const FSkeletonPoseBoneIndex SkeletonBoneIndex(LODBoneIndexToSkeletonBoneIndexMap[LODBoneIndex]);
+				LODBoneIndexToWeightIndexMap[LODBoneIndex] = BlendProfile->GetPerBoneInterpolationIndex(SkeletonBoneIndex, TargetSkeleton, Data.Get());
+			}
 		}
 	}
 	else
 	{
+		// If no blend profile is provided, we use the provided blend per bone data under the assumption it covers a full pose
+		// either using LODBone indexing or full skeleton indexing
+		check(BlendData->PerBoneBlendData.Num() >= NumLODBones);
 		for (int32 LODBoneIndex = 0; LODBoneIndex < NumLODBones; ++LODBoneIndex)
 		{
-			const FSkeletonPoseBoneIndex SkeletonBoneIndex(LODBoneIndexToSkeletonBoneIndexMap[LODBoneIndex]);
-			LODBoneIndexToWeightIndexMap[LODBoneIndex] = BlendProfile->GetPerBoneInterpolationIndex(SkeletonBoneIndex, TargetSkeleton, Data.Get());
+			if (bPerBoneDataUsesSkeletonIndex)
+			{
+				const FSkeletonPoseBoneIndex TargetSkeletonBoneIndex(LODBoneIndexToSkeletonBoneIndexMap[LODBoneIndex]);
+				LODBoneIndexToWeightIndexMap[LODBoneIndex] = TargetSkeletonBoneIndex.GetInt();
+			}
+			else
+			{
+				LODBoneIndexToWeightIndexMap[LODBoneIndex] = LODBoneIndex;
+			}
 		}
 	}
-
+	
 	const float BlendWeight = BlendData->GetClampedWeight();
 
 	if (EnumHasAnyFlags(VM.GetFlags(), EEvaluationFlags::Bones))
@@ -118,6 +149,18 @@ FAnimNextBlendAddKeyframePerBoneWithScaleTask FAnimNextBlendAddKeyframePerBoneWi
 	return Task;
 }
 
+FAnimNextBlendAddKeyframePerBoneWithScaleTask FAnimNextBlendAddKeyframePerBoneWithScaleTask::Make(const FBlendSampleData& BlendDataA, const FBlendSampleData& BlendDataB, float ScaleFactor, bool bPerBoneDataUsesSkeletonIndex)
+{
+	FAnimNextBlendAddKeyframePerBoneWithScaleTask Task;
+	Task.BlendProfile = nullptr;
+	Task.BlendDataA = &BlendDataA;
+	Task.BlendDataB = &BlendDataB;
+	Task.ScaleFactor = ScaleFactor;
+	Task.bPerBoneDataUsesSkeletonIndex = bPerBoneDataUsesSkeletonIndex;
+
+	return Task;
+}
+
 void FAnimNextBlendAddKeyframePerBoneWithScaleTask::Execute(UE::UAF::FEvaluationVM& VM) const
 {
 	using namespace UE::UAF;
@@ -129,7 +172,9 @@ void FAnimNextBlendAddKeyframePerBoneWithScaleTask::Execute(UE::UAF::FEvaluation
 		return;
 	}
 
-	if (BlendProfile == nullptr)
+	if (BlendProfile == nullptr 
+		&& ((BlendDataA->PerBoneBlendData.Num() == 0 || BlendDataB->PerBoneBlendData.Num() == 0)
+			|| (BlendDataA->PerBoneBlendData.Num() != BlendDataB->PerBoneBlendData.Num())))
 	{
 		// If we don't have a blend profile, we blend the whole pose
 		Super::Execute(VM);
@@ -153,32 +198,53 @@ void FAnimNextBlendAddKeyframePerBoneWithScaleTask::Execute(UE::UAF::FEvaluation
 		return;
 	}
 
-	const USkeleton* TargetSkeleton = KeyframeB->Pose.GetSkeletonAsset();
-	TSharedPtr<IInterpolationIndexProvider::FPerBoneInterpolationData> Data = BlendProfile->GetPerBoneInterpolationData(TargetSkeleton);
-
-	const FSkeletonRemapping& SkeletonRemapping = UE::Anim::FSkeletonRemappingRegistry::Get().GetRemapping(BlendProfile->GetSkeleton(), TargetSkeleton);
-
 	const TArrayView<const FBoneIndexType> LODBoneIndexToSkeletonBoneIndexMap = KeyframeB->Pose.GetLODBoneIndexToSkeletonBoneIndexMap();
 	const int32 NumLODBones = LODBoneIndexToSkeletonBoneIndexMap.Num();
 
 	TArray<int32> LODBoneIndexToWeightIndexMap;
 	LODBoneIndexToWeightIndexMap.AddUninitialized(NumLODBones);
 
-	if (SkeletonRemapping.IsValid())
+	if (BlendProfile)
 	{
-		for (int32 LODBoneIndex = 0; LODBoneIndex < NumLODBones; ++LODBoneIndex)
+		const USkeleton* TargetSkeleton = KeyframeB->Pose.GetSkeletonAsset();
+		TSharedPtr<IInterpolationIndexProvider::FPerBoneInterpolationData> Data = BlendProfile->GetPerBoneInterpolationData(TargetSkeleton);
+
+		const FSkeletonRemapping& SkeletonRemapping = UE::Anim::FSkeletonRemappingRegistry::Get().GetRemapping(BlendProfile->GetSkeleton(), TargetSkeleton);
+
+		if (SkeletonRemapping.IsValid())
 		{
-			const FSkeletonPoseBoneIndex TargetSkeletonBoneIndex(LODBoneIndexToSkeletonBoneIndexMap[LODBoneIndex]);
-			const FSkeletonPoseBoneIndex SourceSkeletonBoneIndex(SkeletonRemapping.GetSourceSkeletonBoneIndex(TargetSkeletonBoneIndex.GetInt()));
-			LODBoneIndexToWeightIndexMap[LODBoneIndex] = SourceSkeletonBoneIndex.IsValid() ? BlendProfile->GetPerBoneInterpolationIndex(SourceSkeletonBoneIndex, TargetSkeleton, Data.Get()) : INDEX_NONE;
+			for (int32 LODBoneIndex = 0; LODBoneIndex < NumLODBones; ++LODBoneIndex)
+			{
+				const FSkeletonPoseBoneIndex TargetSkeletonBoneIndex(LODBoneIndexToSkeletonBoneIndexMap[LODBoneIndex]);
+				const FSkeletonPoseBoneIndex SourceSkeletonBoneIndex(SkeletonRemapping.GetSourceSkeletonBoneIndex(TargetSkeletonBoneIndex.GetInt()));
+				LODBoneIndexToWeightIndexMap[LODBoneIndex] = SourceSkeletonBoneIndex.IsValid() ? BlendProfile->GetPerBoneInterpolationIndex(SourceSkeletonBoneIndex, TargetSkeleton, Data.Get()) : INDEX_NONE;
+			}
+		}
+		else
+		{
+			for (int32 LODBoneIndex = 0; LODBoneIndex < NumLODBones; ++LODBoneIndex)
+			{
+				const FSkeletonPoseBoneIndex SkeletonBoneIndex(LODBoneIndexToSkeletonBoneIndexMap[LODBoneIndex]);
+				LODBoneIndexToWeightIndexMap[LODBoneIndex] = BlendProfile->GetPerBoneInterpolationIndex(SkeletonBoneIndex, TargetSkeleton, Data.Get());
+			}
 		}
 	}
 	else
 	{
+		// If no blend profile is provided, we assume the provided bone data encompasses a full pose to be blended together 
+		// either using LOD based indexing or full skeleton indexing
+		check(BlendDataA->PerBoneBlendData.Num() >= NumLODBones);
 		for (int32 LODBoneIndex = 0; LODBoneIndex < NumLODBones; ++LODBoneIndex)
 		{
-			const FSkeletonPoseBoneIndex SkeletonBoneIndex(LODBoneIndexToSkeletonBoneIndexMap[LODBoneIndex]);
-			LODBoneIndexToWeightIndexMap[LODBoneIndex] = BlendProfile->GetPerBoneInterpolationIndex(SkeletonBoneIndex, TargetSkeleton, Data.Get());
+			if (bPerBoneDataUsesSkeletonIndex)
+			{
+				const FSkeletonPoseBoneIndex TargetSkeletonBoneIndex(LODBoneIndexToSkeletonBoneIndexMap[LODBoneIndex]);
+				LODBoneIndexToWeightIndexMap[LODBoneIndex] = TargetSkeletonBoneIndex.GetInt();
+			}
+			else
+			{
+				LODBoneIndexToWeightIndexMap[LODBoneIndex] = LODBoneIndex;
+			}
 		}
 	}
 
@@ -209,7 +275,7 @@ void FAnimNextBlendAddKeyframePerBoneWithScaleTask::Execute(UE::UAF::FEvaluation
 			{ 0, 1 },
 			{ OutputAttributes });
 
-		KeyframeA->Attributes.MoveFrom(OutputAttributes);
+		KeyframeB->Attributes.MoveFrom(OutputAttributes);
 	}
 
 	VM.PushValue(KEYFRAME_STACK_NAME, MoveTemp(KeyframeB));
diff --git a/Engine/Plugins/Experimental/UAF/UAFAnimGraph/Source/UAFAnimGraph/Public/EvaluationVM/Tasks/BlendKeyframesPerBone.h b/Engine/Plugins/Experimental/UAF/UAFAnimGraph/Source/UAFAnimGraph/Public/EvaluationVM/Tasks/BlendKeyframesPerBone.h
--- a/Engine/Plugins/Experimental/UAF/UAFAnimGraph/Source/UAFAnimGraph/Public/EvaluationVM/Tasks/BlendKeyframesPerBone.h
+++ b/Engine/Plugins/Experimental/UAF/UAFAnimGraph/Source/UAFAnimGraph/Public/EvaluationVM/Tasks/BlendKeyframesPerBone.h
@@ -25,7 +25,9 @@ class USkeleton;
  * back the result onto the stack.
  * Top = Top * ScaleFactor
  * 
- * If no blend profile is provided, this task behaves like FAnimNextBlendOverwriteKeyframeWithScaleTask
+ * If no blend profile is provided but the blend data contains valid per bone blend weights,
+ * it will try to use that data to blend each bone together with the assumption the provided data encompasses the full pose.
+ * If neither is provided. this task behaves like FAnimNextBlendOverwriteKeyframeWithScaleTask
  * Note that rotations will not be normalized after this task.
  */
 USTRUCT()
@@ -36,6 +38,7 @@ struct FAnimNextBlendOverwriteKeyframePerBoneWithScaleTask : public FAnimNextBle
 	DECLARE_ANIM_EVALUATION_TASK(FAnimNextBlendOverwriteKeyframePerBoneWithScaleTask)
 
 	static UE_API FAnimNextBlendOverwriteKeyframePerBoneWithScaleTask Make(const IBlendProfileInterface* BlendProfile, const FBlendSampleData& BlendData, float ScaleFactor);
+	static UE_API FAnimNextBlendOverwriteKeyframePerBoneWithScaleTask Make(const FBlendSampleData& BlendData, float ScaleFactor, bool bPerBoneDataUsesSkeletonIndex);
 
 	// Task entry point
 	UE_API virtual void Execute(UE::UAF::FEvaluationVM& VM) const override;
@@ -49,7 +52,9 @@ struct FAnimNextBlendOverwriteKeyframePerBoneWithScaleTask : public FAnimNextBle
 	// An optional blend profile provides a blend weight per bone
 	const IBlendProfileInterface* BlendProfile = nullptr;
 
-	const USkeleton* SourceSkeleton = nullptr;
+	// True if the provided bone weight data uses the full skeleton bone set of the target skeleton for indexing instead of the current LOD's
+	// This is used to build the weight lookup map if no blend profile is used 
+	bool bPerBoneDataUsesSkeletonIndex = false;
 };
 
 /*
@@ -60,7 +65,10 @@ struct FAnimNextBlendOverwriteKeyframePerBoneWithScaleTask : public FAnimNextBle
  * The result is pushed back onto the stack.
  * Top = Top + (Top-1 * ScaleFactor)
  * 
- * If no blend profile is provided, this task behaves like FAnimNextBlendAddKeyframeWithScaleTask
+ * If no blend profile is provided but the blend data contains valid per bone blend weights, 
+ * it will try to use that data to blend each bone together with the assumption the provided data encompasses the full pose. 
+ * If neither is provided, this task behaves like FAnimNextBlendAddKeyframeWithScaleTask
+ * 
  * Note that rotations will not be normalized after this task.
  */
 USTRUCT()
@@ -71,6 +79,7 @@ struct FAnimNextBlendAddKeyframePerBoneWithScaleTask : public FAnimNextBlendAddK
 	DECLARE_ANIM_EVALUATION_TASK(FAnimNextBlendAddKeyframePerBoneWithScaleTask)
 
 	static UE_API FAnimNextBlendAddKeyframePerBoneWithScaleTask Make(const IBlendProfileInterface* BlendProfile, const FBlendSampleData& BlendDataA, const FBlendSampleData& BlendDataB, float ScaleFactor);
+	static UE_API FAnimNextBlendAddKeyframePerBoneWithScaleTask Make(const FBlendSampleData& BlendDataA, const FBlendSampleData& BlendDataB, float ScaleFactor, bool bPerBoneDataUsesSkeletonIndex);
 
 	// Task entry point
 	UE_API virtual void Execute(UE::UAF::FEvaluationVM& VM) const override;
@@ -88,6 +97,10 @@ struct FAnimNextBlendAddKeyframePerBoneWithScaleTask : public FAnimNextBlendAddK
 	// The blend profile to use
 	// An optional blend profile provides a blend weight per bone
 	const IBlendProfileInterface* BlendProfile = nullptr;
+
+	// True if the provided bone weight data uses the full skeleton bone set of the target skeleton for indexing instead of the current LOD's
+	// This is used to build the weight lookup map if no blend profile is used 
+	bool bPerBoneDataUsesSkeletonIndex = false;
 };
 
 
