UAF: Add per bone additive keyframe blend task

diff --git a/Engine/Plugins/Experimental/UAF/UAF/Source/UAF/Private/TransformArrayOperations_AoS.cpp b/Engine/Plugins/Experimental/UAF/UAF/Source/UAF/Private/TransformArrayOperations_AoS.cpp
--- a/Engine/Plugins/Experimental/UAF/UAF/Source/UAF/Private/TransformArrayOperations_AoS.cpp
+++ b/Engine/Plugins/Experimental/UAF/UAF/Source/UAF/Private/TransformArrayOperations_AoS.cpp
@@ -133,6 +133,34 @@ namespace UE::UAF
 		ConvertPoseMeshToLocalRotation(Base, LODBoneIndexToParentLODBoneIndexMap);
 	}
 
+	void BlendWithIdentityAndAccumulatePerBoneWeight(const FTransformArrayAoSView& Base, const FTransformArrayAoSConstView& Additive, 
+		const TArrayView<const int32>& TransformIndexToWeightIndexMap, const TArrayView<float> BoneWeights, const float DefaultBlendWeight)
+	{
+		const int32 NumTransforms = Base.Num();
+		const ScalarRegister VDefaultBlendWeight(DefaultBlendWeight);
+
+		check(Additive.Num() >= NumTransforms);
+		check(TransformIndexToWeightIndexMap.Num() >= NumTransforms);
+
+		for (int32 TransformIndex = 0; TransformIndex < NumTransforms; ++TransformIndex)
+		{
+			const int32 PerTransformIndex = TransformIndexToWeightIndexMap[TransformIndex];
+			const ScalarRegister VBlendWeight = BoneWeights.IsValidIndex(PerTransformIndex) ? ScalarRegister(BoneWeights[PerTransformIndex]) : VDefaultBlendWeight;
+			FTransform::BlendFromIdentityAndAccumulate(Base[TransformIndex], Additive[TransformIndex], VBlendWeight);
+		}
+	}
+
+	void BlendWithIdentityAndAccumulatePerBoneWeightMesh(
+		const FTransformArrayAoSView& Base, const FTransformArrayAoSConstView& Additive, const TArrayView<const FBoneIndexType>& LODBoneIndexToParentLODBoneIndexMap, 
+		const TArrayView<const int32>& TransformIndexToWeightIndexMap, const TArrayView<float> BoneWeights, const float DefaultBlendWeight)
+	{
+		ConvertPoseLocalToMeshRotation(Base, LODBoneIndexToParentLODBoneIndexMap);
+
+		BlendWithIdentityAndAccumulatePerBoneWeight(Base, Additive, TransformIndexToWeightIndexMap, BoneWeights, DefaultBlendWeight);
+
+		ConvertPoseMeshToLocalRotation(Base, LODBoneIndexToParentLODBoneIndexMap);
+	}
+
 	void BlendOverwriteWithScale(const FTransformArrayAoSView& Dest, const FTransformArrayAoSConstView& Source, const float ScaleWeight)
 	{
 		const ScalarRegister VScaleWeight(ScaleWeight);
diff --git a/Engine/Plugins/Experimental/UAF/UAF/Source/UAF/Private/TransformArrayOperations_SoA.cpp b/Engine/Plugins/Experimental/UAF/UAF/Source/UAF/Private/TransformArrayOperations_SoA.cpp
--- a/Engine/Plugins/Experimental/UAF/UAF/Source/UAF/Private/TransformArrayOperations_SoA.cpp
+++ b/Engine/Plugins/Experimental/UAF/UAF/Source/UAF/Private/TransformArrayOperations_SoA.cpp
@@ -257,6 +257,74 @@ namespace UE::UAF
 		ConvertPoseMeshToLocalRotation(Base, LODBoneIndexToParentLODBoneIndexMap);
 	}
 
+	void BlendWithIdentityAndAccumulatePerBoneWeight(const FTransformArraySoAView& Base, const FTransformArraySoAConstView& Additive,
+		const TArrayView<const int32>& TransformIndexToWeightIndexMap, const TArrayView<const float>& BoneWeights, const float DefaultBlendWeight)
+	{
+		const int32 NumTransforms = Base.Num();
+
+		check(Additive.Num() >= NumTransforms);
+		check(TransformIndexToWeightIndexMap.Num() >= NumTransforms);
+
+		using TransformVectorRegister = FTransform::TransformVectorRegister;
+
+		const TransformVectorRegister Zero = VectorZero();
+		const TransformVectorRegister Const0001 = GlobalVectorConstants::Float0001;
+		const TransformVectorRegister ConstNegative0001 = VectorSubtract(Zero, Const0001);
+		const TransformVectorRegister DefaultScale = GlobalVectorConstants::Float1110;
+
+		for (int32 TransformIndex = 0; TransformIndex < NumTransforms; ++TransformIndex)
+		{
+			const int32 PerTransformIndex = TransformIndexToWeightIndexMap[TransformIndex];
+			const float ScaleWeight = BoneWeights.IsValidIndex(PerTransformIndex) ? BoneWeights[PerTransformIndex] : DefaultBlendWeight;
+
+			const TransformVectorRegister VBlendWeight(ScalarRegister(ScaleWeight).Value);
+			const TransformVectorRegister VOneMinusAlpha = VectorSubtract(VectorOne(), VBlendWeight);
+
+			// Blend rotation
+			//     To ensure the 'shortest route', we make sure the dot product between the both rotations is positive.
+			//     const float Bias = (|A.B| >= 0 ? 1 : -1)
+			//     BlendedAtom.Rotation = (B * Alpha) + (A * (Bias * (1.f - Alpha)));
+			//     BlendedAtom.Rotation.QuaternionNormalize();
+			//  Note: A = (0,0,0,1), which simplifies things a lot; only care about sign of B.W now, instead of doing a dot product
+			const TransformVectorRegister RotationB = VectorLoadAligned(&Additive.Rotations[TransformIndex]);
+
+			const TransformVectorRegister QuatRotationDirMask = VectorCompareGE(RotationB, Zero);
+			const TransformVectorRegister BiasTimesA = VectorSelect(QuatRotationDirMask, Const0001, ConstNegative0001);
+			const TransformVectorRegister RotateBTimesWeight = VectorMultiply(RotationB, VBlendWeight);
+			const TransformVectorRegister UnnormalizedRotation = VectorMultiplyAdd(BiasTimesA, VOneMinusAlpha, RotateBTimesWeight);
+
+			// Normalize blended rotation ( result = (Q.Q >= 1e-8) ? (Q / |Q|) : (0,0,0,1) )
+			const TransformVectorRegister BlendedRotation = VectorNormalizeSafe(UnnormalizedRotation, Const0001);
+
+			// FinalAtom.Rotation = BlendedAtom.Rotation * FinalAtom.Rotation;
+			Base.Rotations[TransformIndex] = FQuat::MakeFromVectorRegister(VectorQuaternionMultiply2(BlendedRotation, VectorLoadAligned(&Base.Rotations[TransformIndex])));
+
+			// Blend translation and scale
+			//    BlendedAtom.Translation = Lerp(Zero, Additive.Translation, Alpha);
+			//    BlendedAtom.Scale = Lerp(0, Additive.Scale, Alpha);
+			const TransformVectorRegister BlendedTranslation = FMath::Lerp<TransformVectorRegister>(Zero, VectorLoadFloat3_W0(&Additive.Translations[TransformIndex]), VBlendWeight);
+			const TransformVectorRegister BlendedScale3D = FMath::Lerp<TransformVectorRegister>(Zero, VectorLoadFloat3_W0(&Additive.Scales3D[TransformIndex]), VBlendWeight);
+
+			// Apply translation and scale to final atom
+			//     FinalAtom.Translation += BlendedAtom.Translation
+			//     FinalAtom.Scale *= BlendedAtom.Scale
+			VectorStoreFloat3(VectorAdd(VectorLoadFloat3_W0(&Base.Translations[TransformIndex]), BlendedTranslation), &Base.Translations[TransformIndex]);
+			VectorStoreFloat3(VectorMultiply(VectorLoadFloat3_W0(&Base.Scales3D[TransformIndex]), VectorAdd(DefaultScale, BlendedScale3D)), &Base.Scales3D[TransformIndex]);
+		}
+	}
+
+	void BlendWithIdentityAndAccumulatePerBoneWeightMesh(
+		const FTransformArraySoAView& Base, const FTransformArraySoAConstView& Additive, const TArrayView<const FBoneIndexType>& LODBoneIndexToParentLODBoneIndexMap, 
+		const TArrayView<const int32>& TransformIndexToWeightIndexMap, const TArrayView<const float>& BoneWeights, const float DefaultBlendWeight)
+	{
+		// The additive blend is done in component space for rotations, but local for translation and scale
+		ConvertPoseLocalToMeshRotation(Base, LODBoneIndexToParentLODBoneIndexMap);
+
+		BlendWithIdentityAndAccumulatePerBoneWeight(Base, Additive, TransformIndexToWeightIndexMap, BoneWeights, DefaultBlendWeight);
+
+		ConvertPoseMeshToLocalRotation(Base, LODBoneIndexToParentLODBoneIndexMap);
+	}
+
 	void BlendOverwriteWithScale(const FTransformArraySoAView& Dest, const FTransformArraySoAConstView& Source, const float ScaleWeight)
 	{
 		SCOPE_CYCLE_COUNTER(STAT_AnimNext_BlendOverwrite_SoA);
diff --git a/Engine/Plugins/Experimental/UAF/UAF/Source/UAF/Public/TransformArrayOperations.h b/Engine/Plugins/Experimental/UAF/UAF/Source/UAF/Public/TransformArrayOperations.h
--- a/Engine/Plugins/Experimental/UAF/UAF/Source/UAF/Public/TransformArrayOperations.h
+++ b/Engine/Plugins/Experimental/UAF/UAF/Source/UAF/Public/TransformArrayOperations.h
@@ -54,6 +54,22 @@ namespace UE::UAF
 		const FTransformArraySoAView& Base, const FTransformArraySoAConstView& Additive, 
 		const TArrayView<const FBoneIndexType>& LODBoneIndexToParentLODBoneIndexMap, const float BlendWeight);
 
+	// @TODO: AoS version is untested because AoS currently does not compile in UAF.
+	UAF_API void BlendWithIdentityAndAccumulatePerBoneWeight(const FTransformArrayAoSView& Base, const FTransformArrayAoSConstView& Additive,
+		const TArrayView<const int32>& TransformIndexToWeightIndexMap, const TArrayView<const float>& BoneWeights, const float DefaultBlendWeight);
+	UAF_API void BlendWithIdentityAndAccumulatePerBoneWeight(const FTransformArraySoAView& Base, const FTransformArraySoAConstView& Additive,
+		const TArrayView<const int32>& TransformIndexToWeightIndexMap, const TArrayView<const float>& BoneWeights, const float DefaultBlendWeight);
+
+	// Performs a BlendWithIdentityAndAccumulatePerBoneWeight in mesh space, result is in local space
+	// @TODO: AoS version is untested because AoS currently does not compile in UAF.
+	UAF_API void BlendWithIdentityAndAccumulatePerBoneWeightMesh(
+		const FTransformArrayAoSView& Base, const FTransformArrayAoSConstView& Additive, const TArrayView<const FBoneIndexType>& LODBoneIndexToParentLODBoneIndexMap, 
+		const TArrayView<const int32>& TransformIndexToWeightIndexMap, const TArrayView<const float>& BoneWeights, const float DefaultBlendWeight);
+	// Performs a BlendWithIdentityAndAccumulatePerBoneWeight in mesh space, result is in local space
+	UAF_API void BlendWithIdentityAndAccumulatePerBoneWeightMesh(
+		const FTransformArraySoAView& Base, const FTransformArraySoAConstView& Additive, const TArrayView<const FBoneIndexType>& LODBoneIndexToParentLODBoneIndexMap, 
+		const TArrayView<const int32>& TransformIndexToWeightIndexMap, const TArrayView<const float>& BoneWeights, const float DefaultBlendWeight);
+
 	// The source transforms are scaled by the provided weight and the result is written in the destination
 	// Dest = Source * ScaleWeight
 	UAF_API void BlendOverwriteWithScale(const FTransformArrayAoSView& Dest, const FTransformArrayAoSConstView& Source, const float ScaleWeight);
diff --git a/Engine/Plugins/Experimental/UAF/UAFAnimGraph/Source/UAFAnimGraph/Private/EvaluationVM/Tasks/ApplyAdditiveKeyframePerBone.cpp b/Engine/Plugins/Experimental/UAF/UAFAnimGraph/Source/UAFAnimGraph/Private/EvaluationVM/Tasks/ApplyAdditiveKeyframePerBone.cpp
new file mode 100644
--- /dev/null
+++ b/Engine/Plugins/Experimental/UAF/UAFAnimGraph/Source/UAFAnimGraph/Private/EvaluationVM/Tasks/ApplyAdditiveKeyframePerBone.cpp
@@ -0,0 +1,127 @@
+// Copyright Epic Games, Inc. All Rights Reserved.
+
+#include "EvaluationVM/Tasks/ApplyAdditiveKeyframePerBone.h"
+
+#include "EvaluationVM/EvaluationVM.h"
+#include "EvaluationVM/KeyframeState.h"
+#include "TransformArrayOperations.h"
+#include "Animation/AnimationAsset.h"
+
+#include UE_INLINE_GENERATED_CPP_BY_NAME(ApplyAdditiveKeyframePerBone)
+
+FUAFApplyAdditiveKeyframePerBoneTask FUAFApplyAdditiveKeyframePerBoneTask::Make(const float BlendWeight, const FBlendSampleData& BlendData, bool bPerBoneDataUsesSkeletonIndex)
+{
+	FUAFApplyAdditiveKeyframePerBoneTask Task;
+	Task.BlendData = &BlendData;
+	Task.BlendWeight = BlendWeight;
+	Task.bPerBoneDataUsesSkeletonIndex = bPerBoneDataUsesSkeletonIndex;
+	return Task;
+}
+
+FUAFApplyAdditiveKeyframePerBoneTask FUAFApplyAdditiveKeyframePerBoneTask::Make(const FName& AlphaSourceCurveName, const int8 AlphaCurveInputIndex, TFunction<float(float)> InputScaleBiasClampFn, const FBlendSampleData& BlendData, bool bPerBoneDataUsesSkeletonIndex)
+{
+	FUAFApplyAdditiveKeyframePerBoneTask Task;
+	Task.AlphaSourceCurveName = AlphaSourceCurveName;
+	Task.AlphaCurveInputIndex = AlphaCurveInputIndex;
+	Task.InputScaleBiasClampFn = MoveTemp(InputScaleBiasClampFn);
+	Task.BlendData = &BlendData;
+	Task.bPerBoneDataUsesSkeletonIndex = bPerBoneDataUsesSkeletonIndex;
+	return Task;
+}
+
+void FUAFApplyAdditiveKeyframePerBoneTask::Execute(UE::UAF::FEvaluationVM& VM) const
+{
+	using namespace UE::UAF;
+
+	if (BlendData == nullptr || BlendData->PerBoneBlendData.Num() == 0)
+	{
+		// If we do not have any or enough blend data revert to a full pose additive application
+		Super::Execute(VM);
+		return;
+	}
+
+	if (VM.GetContext().IsValid())
+	{
+		// TODO: Implement with new attribute runtime
+		TUniquePtr<FAttributeCollection> DiscardedCollection;
+		(void)VM.PopValue(ATTRIBUTE_STACK_NAME, DiscardedCollection);
+		return;
+	}
+
+	// Pop our top two poses, the top one being the additive pose and the second one the base
+	TUniquePtr<FKeyframeState> AdditiveKeyframe;
+	if (!VM.PopValue(KEYFRAME_STACK_NAME, AdditiveKeyframe))
+	{
+		// We have no inputs, nothing to do
+		return;
+	}
+
+	TUniquePtr<FKeyframeState> BaseKeyframe;
+	if (!VM.PopValue(KEYFRAME_STACK_NAME, BaseKeyframe))
+	{
+		// We have a single input, discard it since it must be the additive pose, either way something went wrong
+		// Push the reference pose since we'll expect a non-additive pose
+		VM.PushValue(KEYFRAME_STACK_NAME, MakeUnique<FKeyframeState>(VM.MakeReferenceKeyframe(false)));
+		return;
+	}
+
+	if (!AdditiveKeyframe->Pose.IsAdditive())
+	{
+		// Additive must be additive type, push reference pose if not the case
+		VM.PushValue(KEYFRAME_STACK_NAME, MakeUnique<FKeyframeState>(VM.MakeReferenceKeyframe(false)));
+		return;
+	}
+
+	const float AdditiveWeight = GetInterpolationAlpha(BaseKeyframe.Get(), AdditiveKeyframe.Get());
+
+	const TArrayView<const FBoneIndexType> LODBoneIndexToSkeletonBoneIndexMap = BaseKeyframe->Pose.GetLODBoneIndexToSkeletonBoneIndexMap();
+	const int32 NumLODBones = LODBoneIndexToSkeletonBoneIndexMap.Num();
+
+	TArray<int32> LODBoneIndexToWeightIndexMap;
+	LODBoneIndexToWeightIndexMap.AddUninitialized(NumLODBones);
+
+	for (int32 LODBoneIndex = 0; LODBoneIndex < NumLODBones; ++LODBoneIndex)
+	{
+		if (bPerBoneDataUsesSkeletonIndex)
+		{
+			const FSkeletonPoseBoneIndex TargetSkeletonBoneIndex(LODBoneIndexToSkeletonBoneIndexMap[LODBoneIndex]);
+			LODBoneIndexToWeightIndexMap[LODBoneIndex] = TargetSkeletonBoneIndex.GetInt();
+		}
+		else
+		{
+			LODBoneIndexToWeightIndexMap[LODBoneIndex] = LODBoneIndex;
+		}
+	}
+
+	if (EnumHasAnyFlags(VM.GetFlags(), EEvaluationFlags::Bones))
+	{
+		check(BaseKeyframe->Pose.GetNumBones() == AdditiveKeyframe->Pose.GetNumBones());
+
+		const FTransformArrayView BaseTransformsView = BaseKeyframe->Pose.LocalTransforms.GetView();
+
+		if (AdditiveKeyframe->Pose.IsMeshSpaceAdditive())
+		{
+			BlendWithIdentityAndAccumulatePerBoneWeightMesh(
+				BaseTransformsView, AdditiveKeyframe->Pose.LocalTransforms.GetConstView(),
+				AdditiveKeyframe->Pose.GetLODBoneIndexToParentLODBoneIndexMap(), LODBoneIndexToWeightIndexMap, BlendData->PerBoneBlendData, AdditiveWeight);
+		}
+		else
+		{
+			BlendWithIdentityAndAccumulatePerBoneWeight(BaseTransformsView, AdditiveKeyframe->Pose.LocalTransforms.GetConstView(), LODBoneIndexToWeightIndexMap, BlendData->PerBoneBlendData, AdditiveWeight);
+		}
+
+		NormalizeRotations(BaseTransformsView);
+	}
+
+	if (EnumHasAnyFlags(VM.GetFlags(), EEvaluationFlags::Curves))
+	{
+		BaseKeyframe->Curves.Accumulate(AdditiveKeyframe->Curves, AdditiveWeight);
+	}
+
+	if (EnumHasAnyFlags(VM.GetFlags(), EEvaluationFlags::Attributes))
+	{
+		UE::Anim::Attributes::AccumulateAttributes(AdditiveKeyframe->Attributes, BaseKeyframe->Attributes, AdditiveWeight, AAT_None);
+	}
+
+	VM.PushValue(KEYFRAME_STACK_NAME, MoveTemp(BaseKeyframe));
+}
diff --git a/Engine/Plugins/Experimental/UAF/UAFAnimGraph/Source/UAFAnimGraph/Public/EvaluationVM/Tasks/ApplyAdditiveKeyframePerBone.h b/Engine/Plugins/Experimental/UAF/UAFAnimGraph/Source/UAFAnimGraph/Public/EvaluationVM/Tasks/ApplyAdditiveKeyframePerBone.h
new file mode 100644
--- /dev/null
+++ b/Engine/Plugins/Experimental/UAF/UAFAnimGraph/Source/UAFAnimGraph/Public/EvaluationVM/Tasks/ApplyAdditiveKeyframePerBone.h
@@ -0,0 +1,39 @@
+// Copyright Epic Games, Inc. All Rights Reserved.
+
+#pragma once
+
+#include "CoreMinimal.h"
+#include "EvaluationVM/Tasks/ApplyAdditiveKeyframe.h"
+
+#include "ApplyAdditiveKeyframePerBone.generated.h"
+
+#define UE_API UAFANIMGRAPH_API
+
+struct FBlendSampleData;
+
+/*
+ * Apply Additive Keyframe Per Bone Task
+ *
+ * This pops the top two keyframes from the VM keyframe stack, it applies an additive keyframe with a weight per bone onto its base, 
+ * and pushes back the result onto the stack.
+ * If no valid blend data is provided it will revert back to a full pose additive application. 
+ * 
+ * The top pose should be the additive keyframe and the second to the top the base keyframe.
+ */
+USTRUCT()
+struct FUAFApplyAdditiveKeyframePerBoneTask : public FAnimNextApplyAdditiveKeyframeTask
+{
+	GENERATED_BODY()
+
+	DECLARE_ANIM_EVALUATION_TASK(FUAFApplyAdditiveKeyframePerBoneTask)
+
+	static UE_API FUAFApplyAdditiveKeyframePerBoneTask Make(const float BlendWeight, const FBlendSampleData& BlendData, bool bPerBoneDataUsesSkeletonIndex);
+	static UE_API FUAFApplyAdditiveKeyframePerBoneTask Make(const FName& AlphaSourceCurveName, const int8 AlphaCurveInputIndex, TFunction<float(float)> InputScaleBiasClampFn, const FBlendSampleData& BlendData, bool bPerBoneDataUsesSkeletonIndex);
+
+	UE_API virtual void Execute(UE::UAF::FEvaluationVM& VM) const override;
+
+	const FBlendSampleData* BlendData = nullptr;
+	bool bPerBoneDataUsesSkeletonIndex = false;
+};
+
+#undef UE_API
