Fix Blend Stack Requester trait not implementing IGarbageCollection

diff --git a/Engine/Plugins/Experimental/UAF/UAFAnimGraph/Source/UAFAnimGraph/Internal/Traits/BlendStackTrait.h b/Engine/Plugins/Experimental/UAF/UAFAnimGraph/Source/UAFAnimGraph/Internal/Traits/BlendStackTrait.h
--- a/Engine/Plugins/Experimental/UAF/UAFAnimGraph/Source/UAFAnimGraph/Internal/Traits/BlendStackTrait.h
+++ b/Engine/Plugins/Experimental/UAF/UAFAnimGraph/Source/UAFAnimGraph/Internal/Traits/BlendStackTrait.h
@@ -238,7 +238,7 @@ namespace UE::UAF
 	 * Example additive trait used to push graphs to a Blend Stack base trait.
 	 * Same behavior as FBlendStackTrait, but as additive (i.e. FBlendStackCoreTrait base + FBlendStackRequesterTrait as additive).
 	 */	
-	struct FBlendStackRequesterTrait : FAdditiveTrait, IUpdate
+	struct FBlendStackRequesterTrait : FAdditiveTrait, IUpdate, IGarbageCollection
 	{
 		DECLARE_ANIM_TRAIT(FBlendStackRequesterTrait, FAdditiveTrait)
 
@@ -246,5 +246,8 @@ namespace UE::UAF
 
 		// IUpdate impl
 		virtual void PreUpdate(FUpdateTraversalContext& Context, const TTraitBinding<IUpdate>& Binding, const FTraitUpdateState& TraitState) const override;
+
+		// IGarbageCollection impl
+		virtual void AddReferencedObjects(const FExecutionContext& Context, const TTraitBinding<IGarbageCollection>& Binding, FReferenceCollector& Collector) const override;
 	};
 }
diff --git a/Engine/Plugins/Experimental/UAF/UAFAnimGraph/Source/UAFAnimGraph/Private/Traits/BlendStackTrait.cpp b/Engine/Plugins/Experimental/UAF/UAFAnimGraph/Source/UAFAnimGraph/Private/Traits/BlendStackTrait.cpp
--- a/Engine/Plugins/Experimental/UAF/UAFAnimGraph/Source/UAFAnimGraph/Private/Traits/BlendStackTrait.cpp
+++ b/Engine/Plugins/Experimental/UAF/UAFAnimGraph/Source/UAFAnimGraph/Private/Traits/BlendStackTrait.cpp
@@ -45,14 +45,19 @@ namespace UE::UAF
 		GeneratorMacro(IUpdate) \
 
 	GENERATE_ANIM_TRAIT_IMPLEMENTATION(FBlendStackTrait, TRAIT_INTERFACE_ENUMERATOR, NULL_ANIM_TRAIT_INTERFACE_ENUMERATOR, NULL_ANIM_TRAIT_EVENT_ENUMERATOR)
+	#undef TRAIT_INTERFACE_ENUMERATOR
+
+	#define TRAIT_INTERFACE_ENUMERATOR(GeneratorMacro) \
+		GeneratorMacro(IUpdate) \
+		GeneratorMacro(IGarbageCollection) \
 
 		// Trait required interfaces implementation boilerplate
-#define TRAIT_REQUIRED_INTERFACE_ENUMERATOR(GeneratorMacroRequired) \
+	#define TRAIT_REQUIRED_INTERFACE_ENUMERATOR(GeneratorMacroRequired) \
 		GeneratorMacroRequired(IBlendStack) \
 
 		GENERATE_ANIM_TRAIT_IMPLEMENTATION(FBlendStackRequesterTrait, TRAIT_INTERFACE_ENUMERATOR, TRAIT_REQUIRED_INTERFACE_ENUMERATOR, NULL_ANIM_TRAIT_EVENT_ENUMERATOR)
-#undef TRAIT_INTERFACE_ENUMERATOR
-#undef TRAIT_REQUIRED_INTERFACE_ENUMERATOR
+	#undef TRAIT_INTERFACE_ENUMERATOR
+	#undef TRAIT_REQUIRED_INTERFACE_ENUMERATOR
 
 
 	
@@ -563,4 +568,14 @@ namespace UE::UAF
 			}
 		}
 	}
+
+	void FBlendStackRequesterTrait::AddReferencedObjects(const FExecutionContext& Context, const TTraitBinding<IGarbageCollection>& Binding, FReferenceCollector& Collector) const
+	{
+		IGarbageCollection::AddReferencedObjects(Context, Binding, Collector);
+
+		const FSharedData* SharedData = Binding.GetSharedData<FSharedData>();
+
+		TObjectPtr<const UObject> Object = SharedData->GetObject(Binding);
+		Collector.AddReferencedObject(Object);
+	}
 }
