Fixed RigVM error caused by a missing struct in the allowed list

diff --git a/Engine/Plugins/Experimental/UAF/UAFAnimGraph/Source/UAFAnimGraph/Private/AnimNextAnimGraphModule.cpp b/Engine/Plugins/Experimental/UAF/UAFAnimGraph/Source/UAFAnimGraph/Private/AnimNextAnimGraphModule.cpp
--- a/Engine/Plugins/Experimental/UAF/UAFAnimGraph/Source/UAFAnimGraph/Private/AnimNextAnimGraphModule.cpp
+++ b/Engine/Plugins/Experimental/UAF/UAFAnimGraph/Source/UAFAnimGraph/Private/AnimNextAnimGraphModule.cpp
@@ -29,6 +29,8 @@
 #include "Native/UAFSimplePrePhysicsGraphComponent.h"
 #include "Traits/BlendSmoother.h"
 #include "Traits/BlendStackTrait.h"
+#include "Graph/AnimNextAnimGraph.h"
+
 
 #if WITH_ANIMNEXT_CONSOLE_COMMANDS
 #include "HAL/IConsoleManager.h"
@@ -67,6 +69,7 @@ void FAnimNextAnimGraphModule::StartupModule()
 	
 	static UScriptStruct* const AllowedStructTypes[] =
 	{
+		FAnimNextAnimGraph::StaticStruct(),
 		FRigVMTrait_AnimNextPublicVariables::StaticStruct(),
 		FAnimNextInjectionBlendSettings::StaticStruct(),
 		FModifyCurveParameters::StaticStruct(),
