UAF: Add blend profile support to the montage trait

diff --git a/Engine/Plugins/Experimental/UAF/UAFAnimGraph/Source/UAFAnimGraph/Internal/Injection/MontageTrait.h b/Engine/Plugins/Experimental/UAF/UAFAnimGraph/Source/UAFAnimGraph/Internal/Injection/MontageTrait.h
--- a/Engine/Plugins/Experimental/UAF/UAFAnimGraph/Source/UAFAnimGraph/Internal/Injection/MontageTrait.h
+++ b/Engine/Plugins/Experimental/UAF/UAFAnimGraph/Source/UAFAnimGraph/Internal/Injection/MontageTrait.h
@@ -62,19 +62,35 @@ namespace UE::UAF
 
 		struct FMontageEvalTaskData
 		{
-			float AnimPosition = 0.0f;
-			TWeakObjectPtr<UAnimSequence> WeakAnimationSequence = nullptr;
-
 			FMontageEvalTaskData(float InAnimPosition, UAnimSequence* InAnimSequence)
 				: AnimPosition(InAnimPosition)
 				, WeakAnimationSequence(InAnimSequence)
-			{ }
+			{
+			}
+
+			float AnimPosition = 0.0f;
+			TWeakObjectPtr<UAnimSequence> WeakAnimationSequence = nullptr;
 		};
 
 		struct FInstanceData : FTrait::FInstanceData
 		{
-			/** Input node from which we receive the source pose to blend montages against during blend windows */
-			FTraitPtr Source;
+			/** If the pre-update or post-update state should be returned on GetState() */
+			bool bReturnPreUpdateState = true;
+
+			/**
+			 * If any of the active montages are using blend profiles,
+			 * this will affect what blend is performed and what blend data is created
+			 */
+			bool bUsesBlendProfiles = false;
+
+			/** Cached Ptr to the montage component */
+			FUAFMontageComponent* MontageComponent = nullptr;
+
+			/** The eval data index into the montage component for the active synchronization candidate */
+			int32 SyncCandidateIndex = INDEX_NONE;
+
+			/** Delta in seconds between updates, populated during PreUpdate */
+			float DeltaTime = 0.f;
 
 			/** The current blend weight of the source input */
 			float SourceWeight = 1.0f;
@@ -82,8 +98,17 @@ namespace UE::UAF
 			/** The source weight of the previous update, used to calculate if the source is blending out */
 			float LastUpdateSourceWeight = 0.0f;
 
-			/** Delta in seconds between updates, populated during PreUpdate */
-			float DeltaTime = 0.f;
+			/** Input node from which we receive the source pose to blend montages against during blend windows */
+			FTraitPtr Source;
+
+			/** The state representing the timeline before the montage update */
+			FTimelineState PreMontageUpdateState = FTimelineState();
+
+			/** The state representing the timeline after the montage update */
+			FTimelineState PostMontageUpdateState = FTimelineState();
+
+			/** The sync group params used to participate in synchronization */
+			FSyncGroupParameters SyncGroupParams;
 
 			/** The individual weights for each montage animation, plus source if relevant */
 			TArray<float> PoseWeights;
@@ -97,25 +122,14 @@ namespace UE::UAF
 			/** The additive anim data needed for the evaluation task for this slot */
 			TArray<FMontageEvalTaskData> AdditiveEvalTaskDataForSlot;
 
-			/** The state representing the timeline before the montage update */
-			FTimelineState PreMontageUpdateState = FTimelineState();
-
-			/** The state representing the timeline after the montage update */
-			FTimelineState PostMontageUpdateState = FTimelineState();
-
-			/** If the pre-update or post-update state should be returned on GetState() */
-			bool bReturnPreUpdateState = true;
-
-			/** The sync group params used to participate in synchronization */
-			FSyncGroupParameters SyncGroupParams;
-
-			/** The eval data index into the montage component for the active synchronization candidate */
-			int32 SyncCandidateIndex = INDEX_NONE;
+			/** Blend data for non-additive poses if blend profiles are used */
+			TArray<FBlendSampleData> BlendDataForBlendProfiles;
 
-			/** Cached Ptr to the montage component */
-			FUAFMontageComponent* MontageComponent = nullptr;
+			/** Blend data for additive poses if blend profiles are used */
+			TArray<FBlendSampleData> AdditiveBlendDataForBlendProfiles;
 
 			void Construct(const FExecutionContext& Context, const FTraitBinding& Binding);
+			void ResetData();
 		};
 
 		using FSharedData = FUAFMontageTraitSharedData;
diff --git a/Engine/Plugins/Experimental/UAF/UAFAnimGraph/Source/UAFAnimGraph/Internal/Injection/UAFMontageComponent.h b/Engine/Plugins/Experimental/UAF/UAFAnimGraph/Source/UAFAnimGraph/Internal/Injection/UAFMontageComponent.h
--- a/Engine/Plugins/Experimental/UAF/UAFAnimGraph/Source/UAFAnimGraph/Internal/Injection/UAFMontageComponent.h
+++ b/Engine/Plugins/Experimental/UAF/UAFAnimGraph/Source/UAFAnimGraph/Internal/Injection/UAFMontageComponent.h
@@ -36,11 +36,13 @@ struct FUAFMontageEvalData
 
 public:
 	FUAFMontageEvalData() = default;
-	FUAFMontageEvalData(UAnimMontage* InMontage, float InPosition, const FAlphaBlend& InBlendInfo, float InPlayRate)
+	FUAFMontageEvalData(UAnimMontage* InMontage, float InPosition, float InPlayRate, float InBlendStartAlpha, const FAlphaBlend& InBlendInfo, const UBlendProfile* InBlendProfile)
 		: Montage(InMontage)
 		, BlendInfo(InBlendInfo)
+		, ActiveBlendProfile(InBlendProfile)
 		, MontagePosition(InPosition)
 		, MontagePlayRate(InPlayRate)
+		, BlendStartAlpha(InBlendStartAlpha)
 	{
 	}
 
@@ -51,8 +53,12 @@ private:
 	UPROPERTY(Transient)
 	FAlphaBlend BlendInfo = FAlphaBlend();
 
+	UPROPERTY(Transient)
+	TObjectPtr<const UBlendProfile> ActiveBlendProfile = nullptr;
+
 	float MontagePosition = 0.0f;
 	float MontagePlayRate = 1.0f;
+	float BlendStartAlpha = 0.0f;
 
 	friend struct UE::UAF::FMontageTrait;
 };
diff --git a/Engine/Plugins/Experimental/UAF/UAFAnimGraph/Source/UAFAnimGraph/Private/Injection/MontageTrait.cpp b/Engine/Plugins/Experimental/UAF/UAFAnimGraph/Source/UAFAnimGraph/Private/Injection/MontageTrait.cpp
--- a/Engine/Plugins/Experimental/UAF/UAFAnimGraph/Source/UAFAnimGraph/Private/Injection/MontageTrait.cpp
+++ b/Engine/Plugins/Experimental/UAF/UAFAnimGraph/Source/UAFAnimGraph/Private/Injection/MontageTrait.cpp
@@ -18,6 +18,8 @@
 #include "Injection/UAFMontageComponent.h"
 #include "Graph/SyncGroup_GraphInstanceComponent.h"
 #include "EvaluationVM/Tasks/ApplyAdditiveKeyframe.h"
+#include "EvaluationVM/Tasks/BlendKeyframesPerBone.h"
+#include "EvaluationVM/Tasks/ApplyAdditiveKeyframePerBone.h"
 
 namespace UE::UAF
 {
@@ -59,20 +61,27 @@ namespace UE::UAF
 		InstanceData->SyncGroupParams.GroupRole = EAnimGroupSynchronizationRole::AlwaysLeader;
 	}
 
+	void FMontageTrait::FInstanceData::ResetData()
+	{
+		SourceWeight = 1.0f;
+		PoseWeights.Reset();
+		AdditivePoseWeights.Reset();
+		EvalTaskDataForSlot.Reset();
+		AdditiveEvalTaskDataForSlot.Reset();
+		BlendDataForBlendProfiles.Reset();
+		AdditiveBlendDataForBlendProfiles.Reset();
+		bReturnPreUpdateState = true;
+		SyncCandidateIndex = INDEX_NONE;
+	}
+
 	void FMontageTrait::PreUpdate(FUpdateTraversalContext& Context, const TTraitBinding<IUpdate>& Binding, const FTraitUpdateState& TraitState) const
 	{
 		IUpdate::PreUpdate(Context, Binding, TraitState);
 
 		FInstanceData* InstanceData = Binding.GetInstanceData<FInstanceData>();
 		InstanceData->DeltaTime = TraitState.GetDeltaTime();
-		InstanceData->PoseWeights.Reset();
-		InstanceData->EvalTaskDataForSlot.Reset();
-		InstanceData->AdditiveEvalTaskDataForSlot.Reset();
-		InstanceData->AdditivePoseWeights.Reset();
 		InstanceData->LastUpdateSourceWeight = InstanceData->SourceWeight;
-		InstanceData->SourceWeight = 1.0f;
-		InstanceData->bReturnPreUpdateState = true;
-		InstanceData->SyncCandidateIndex = INDEX_NONE;
+		InstanceData->ResetData();
 
 		const FUAFMontageComponent* MontageComponent = InstanceData->MontageComponent;
 		if (MontageComponent == nullptr)
@@ -92,14 +101,34 @@ namespace UE::UAF
 			return;
 		}
 
+		InstanceData->bUsesBlendProfiles = false;
+		for (const FUAFMontageEvalData& EvalData : MontageComponent->MontageEvalData)
+		{
+			if (EvalData.ActiveBlendProfile != nullptr)
+			{
+				InstanceData->bUsesBlendProfiles = true;
+				break;
+			}
+		}
+
+		int32 NumBones = 0;
+		if (InstanceData->bUsesBlendProfiles)
+		{
+			const TWeakObjectPtr<const USkeletalMeshComponent> RefSkelMeshComp = Context.GetBindingObject();
+			if (RefSkelMeshComp.IsValid())
+			{
+				NumBones = RefSkelMeshComp->GetNumBones();
+			}
+		}
+		
 		const bool bIsSyncDisabled = SharedData->GetbDisableSynchronization(Binding);
 		float TotalMontageWeight = 0.0f;
 		float TotalNonAdditiveMontageWeight = 0.0f;
 		float SyncCandidate_HeighestWeight = -1.0f;
 		int32 SyncCandidate_HeighestWeightIndex = INDEX_NONE;
-		for (int32 i = 0; i < MontageComponent->MontageEvalData.Num(); ++i)
+		for (int32 EvalDataIndex = 0; EvalDataIndex < MontageComponent->MontageEvalData.Num(); ++EvalDataIndex)
 		{
-			const FUAFMontageEvalData& EvalData = MontageComponent->MontageEvalData[i];
+			const FUAFMontageEvalData& EvalData = MontageComponent->MontageEvalData[EvalDataIndex];
 			if (EvalData.Montage != nullptr)
 			{
 				const UAnimMontage* const Montage = EvalData.Montage.Get();
@@ -108,23 +137,54 @@ namespace UE::UAF
 					if (const FAnimTrack* AnimTrack = EvalData.Montage->GetAnimationData(CurrentSlotName))
 					{
 						const float ClampedTime = FMath::Clamp(EvalData.MontagePosition, 0.f, AnimTrack->GetLength());
-						if (const FAnimSegment* AnimSegment = AnimTrack->GetSegmentAtTime(EvalData.MontagePosition))
+						if (const FAnimSegment* AnimSegment = AnimTrack->GetSegmentAtTime(ClampedTime))
 						{
 							float AnimPosition = 0;
 							if (UAnimSequence* AnimSegmentSequence = Cast<UAnimSequence>(AnimSegment->GetAnimationData(ClampedTime, AnimPosition)))
 							{
 								const float MontageWeight = EvalData.BlendInfo.GetBlendedValue();
 								TotalMontageWeight += MontageWeight;
+								
 								if (AnimTrack->IsAdditive())
 								{
-									InstanceData->AdditivePoseWeights.Add(MontageWeight);
 									InstanceData->AdditiveEvalTaskDataForSlot.Emplace(AnimPosition, AnimSegmentSequence);
+									if (!InstanceData->bUsesBlendProfiles)
+									{
+										InstanceData->AdditivePoseWeights.Add(MontageWeight);
+									}
 								}
 								else
 								{
-									InstanceData->PoseWeights.Add(MontageWeight);
 									InstanceData->EvalTaskDataForSlot.Emplace(AnimPosition, AnimSegmentSequence);
 									TotalNonAdditiveMontageWeight += MontageWeight;
+									if (!InstanceData->bUsesBlendProfiles)
+									{
+										InstanceData->PoseWeights.Add(MontageWeight);
+									}
+								}
+
+								if (InstanceData->bUsesBlendProfiles)
+								{
+									FBlendSampleData& SampleBlendData = AnimTrack->IsAdditive() ? InstanceData->AdditiveBlendDataForBlendProfiles.AddDefaulted_GetRef() : InstanceData->BlendDataForBlendProfiles.AddDefaulted_GetRef();
+									SampleBlendData.TotalWeight = MontageWeight;
+									SampleBlendData.PerBoneBlendData.AddZeroed(NumBones);
+
+									const UBlendProfile* BlendProfile = EvalData.ActiveBlendProfile;
+									for (int32 PerBoneIndex = 0; PerBoneIndex < SampleBlendData.PerBoneBlendData.Num(); ++PerBoneIndex)
+									{
+										// If we have a active blend profile we calculate the bone weight according to its settings
+										// Otherwise we fall back to use the pose weight as bone weight
+										if (BlendProfile)
+										{
+											const float BoneBlendScale = BlendProfile->GetBoneBlendScale(PerBoneIndex);
+											const float BoneWeight = UBlendProfile::CalculateBoneWeight(BoneBlendScale, BlendProfile->GetMode(), EvalData.BlendInfo, EvalData.BlendStartAlpha, MontageWeight, false);
+											SampleBlendData.PerBoneBlendData[PerBoneIndex] = BoneWeight;
+										}
+										else
+										{
+											SampleBlendData.PerBoneBlendData[PerBoneIndex] = MontageWeight;
+										}
+									}
 								}
 
 								if (!bIsSyncDisabled)
@@ -133,7 +193,7 @@ namespace UE::UAF
 									if (bIsSynchronizationCandidate && MontageWeight > SyncCandidate_HeighestWeight)
 									{
 										SyncCandidate_HeighestWeight = MontageWeight;
-										SyncCandidate_HeighestWeightIndex = i;
+										SyncCandidate_HeighestWeightIndex = EvalDataIndex;
 									}
 								}
 							}
@@ -143,31 +203,104 @@ namespace UE::UAF
 			}
 		}
 
-		if (TotalMontageWeight > (1.f + ZERO_ANIMWEIGHT_THRESH))
+		// Re-Normalize pose weights 
+		if (InstanceData->bUsesBlendProfiles)
 		{
-			// If we have to re-normalize the weights, we also have to recalculate the total non-additive weight 
-			TotalNonAdditiveMontageWeight = 0.0f;
-			for (int32 Index = 0; Index < InstanceData->PoseWeights.Num(); ++Index)
+			if (TotalNonAdditiveMontageWeight > (1.0f + ZERO_ANIMWEIGHT_THRESH))
 			{
-				InstanceData->PoseWeights[Index] /= TotalMontageWeight;
-				TotalNonAdditiveMontageWeight += InstanceData->PoseWeights[Index];
+				for (FBlendSampleData& BlendSampleEntry : InstanceData->BlendDataForBlendProfiles)
+				{
+					BlendSampleEntry.TotalWeight /= TotalNonAdditiveMontageWeight;
+				}
 			}
+		}
+		else
+		{
+			if (TotalMontageWeight > (1.f + ZERO_ANIMWEIGHT_THRESH))
+			{
+				// If we have to re-normalize the weights, we also have to recalculate the total non-additive weight 
+				TotalNonAdditiveMontageWeight = 0.0f;
+				for (int32 Index = 0; Index < InstanceData->PoseWeights.Num(); ++Index)
+				{
+					InstanceData->PoseWeights[Index] /= TotalMontageWeight;
+					TotalNonAdditiveMontageWeight += InstanceData->PoseWeights[Index];
+				}
+
+				// To maintain parity with the current ABP implementation we normalize additives when no blend profiles are used 
+				for (int32 Index = 0; Index < InstanceData->AdditivePoseWeights.Num(); ++Index)
+				{
+					InstanceData->AdditivePoseWeights[Index] /= TotalMontageWeight;
+				}
+			}	
+		}	
+		
+		// Calculate weight for the source pose 
+		InstanceData->SourceWeight = 1.0f - FMath::Clamp(TotalNonAdditiveMontageWeight, 0.0f, 1.0f);
 
-			for (int32 Index = 0; Index < InstanceData->AdditivePoseWeights.Num(); ++Index)
+		TArray<float> TotalBoneWeights;
+		TotalBoneWeights.SetNumZeroed(NumBones);
+		if (InstanceData->bUsesBlendProfiles)
+		{
+			// Calculate non-additive total bone weights 
+			for (const FBlendSampleData& BlendSampleEntry : InstanceData->BlendDataForBlendProfiles)
 			{
-				InstanceData->AdditivePoseWeights[Index] /= TotalMontageWeight;
+				for (int32 PerBoneIndex = 0; PerBoneIndex < NumBones; ++PerBoneIndex)
+				{
+					if (ensureMsgf(BlendSampleEntry.PerBoneBlendData.IsValidIndex(PerBoneIndex), TEXT("Montage Bone Weight Normalization: Invalid per bone blend data setup!")))
+					{
+						TotalBoneWeights[PerBoneIndex] += BlendSampleEntry.PerBoneBlendData[PerBoneIndex];
+					}
+				}
 			}
 
-		}
+			// Re-Normalize non-additive bone weights 
+			const float NormalizeThreshold = FAnimWeight::IsRelevant(InstanceData->SourceWeight) ? (1.f + ZERO_ANIMWEIGHT_THRESH) : ZERO_ANIMWEIGHT_THRESH;
+			for (int32 PerBoneIndex = 0; PerBoneIndex < NumBones; ++PerBoneIndex)
+			{
+				float NormalizedBoneTotal = 0.0f;
+				const float TotalBoneWeight = TotalBoneWeights[PerBoneIndex];
+				const int32 NumBlendSamples = InstanceData->BlendDataForBlendProfiles.Num();
+				for (FBlendSampleData& BlendSampleEntry : InstanceData->BlendDataForBlendProfiles)
+				{
+					if (TotalBoneWeight > NormalizeThreshold)
+					{
+						BlendSampleEntry.PerBoneBlendData[PerBoneIndex] /= TotalBoneWeight;
+					}
+					else
+					{
+						if (!FAnimWeight::IsRelevant(InstanceData->SourceWeight))
+						{
+							BlendSampleEntry.PerBoneBlendData[PerBoneIndex] = 1.0f / NumBlendSamples;
+						}
+					}
+
+					NormalizedBoneTotal += BlendSampleEntry.PerBoneBlendData[PerBoneIndex];
+				}
 
+				TotalBoneWeights[PerBoneIndex] = NormalizedBoneTotal;
+			}
+		}
+		
 		// If we are only playing additive animations we do not add the source weight as we want to apply 100% of it and not blend it in
 		const bool bOnlyPlayingAdditives = InstanceData->EvalTaskDataForSlot.Num() == 0 && InstanceData->AdditiveEvalTaskDataForSlot.Num() > 0;
 		const bool bAlwaysUpdateSource = SharedData->GetbAlwaysUpdateSource(Binding);
-		InstanceData->SourceWeight = 1.0f - FMath::Clamp(TotalNonAdditiveMontageWeight, 0.0f, 1.0f);
 		if ((FAnimWeight::IsRelevant(InstanceData->SourceWeight) || bAlwaysUpdateSource) && !bOnlyPlayingAdditives)
 		{
 			// Add the pose weight to blend in the source with
-			InstanceData->PoseWeights.Add(InstanceData->SourceWeight);
+			if (InstanceData->bUsesBlendProfiles)
+			{
+				FBlendSampleData& SourceSampleBlendData = InstanceData->BlendDataForBlendProfiles.AddDefaulted_GetRef();
+				SourceSampleBlendData.TotalWeight = InstanceData->SourceWeight;
+				SourceSampleBlendData.PerBoneBlendData.SetNumZeroed(NumBones);
+				for (int32 PerBoneIndex = 0; PerBoneIndex < NumBones; ++PerBoneIndex)
+				{
+					SourceSampleBlendData.PerBoneBlendData[PerBoneIndex] = 1.0f - TotalBoneWeights[PerBoneIndex];
+				}
+			}
+			else
+			{
+				InstanceData->PoseWeights.Add(InstanceData->SourceWeight);
+			}
 		}
 
 		// Queue up anim notify events
@@ -240,9 +373,9 @@ namespace UE::UAF
 
 		// First we push all non-additive animation keyframes onto the stack
 		const float DeltaTime = InstanceData->DeltaTime;
-		for (int32 iEvalTaskData = InstanceData->EvalTaskDataForSlot.Num() -1; iEvalTaskData >= 0; --iEvalTaskData)
+		for (int32 EvalTaskDataIndex = InstanceData->EvalTaskDataForSlot.Num() -1; EvalTaskDataIndex >= 0; --EvalTaskDataIndex)
 		{
-			const FMontageEvalTaskData& TaskData = InstanceData->EvalTaskDataForSlot[iEvalTaskData];
+			const FMontageEvalTaskData& TaskData = InstanceData->EvalTaskDataForSlot[EvalTaskDataIndex];
 			if (TaskData.WeakAnimationSequence.IsValid())
 			{
 				FAnimNextAnimSequenceKeyframeTask Task;
@@ -258,39 +391,74 @@ namespace UE::UAF
 		}
 
 		// Then we blend the non-additive poses and source pose (if applicable) together 
-		if (InstanceData->PoseWeights.Num() > 0)
+		if (InstanceData->bUsesBlendProfiles)
 		{
-			Context.AppendTask(FAnimNextBlendOverwriteKeyframeWithScaleTask::Make(InstanceData->PoseWeights[0]));
-			for (int32 i = 1; i < InstanceData->PoseWeights.Num(); ++i)
+			if (InstanceData->BlendDataForBlendProfiles.Num() > 0)
 			{
-				Context.AppendTask(FAnimNextBlendAddKeyframeWithScaleTask::Make(InstanceData->PoseWeights[i]));
+				const FBlendSampleData& SampleBlendData = InstanceData->BlendDataForBlendProfiles[0];
+				Context.AppendTask(FAnimNextBlendOverwriteKeyframePerBoneWithScaleTask::Make(SampleBlendData, SampleBlendData.GetClampedWeight(), true));
+
+				for (int32 BlendDataIndex = 1; BlendDataIndex < InstanceData->BlendDataForBlendProfiles.Num(); ++BlendDataIndex)
+				{
+					const FBlendSampleData& SampleBlendDataA = InstanceData->BlendDataForBlendProfiles[BlendDataIndex];
+					const FBlendSampleData& SampleBlendDataB = InstanceData->BlendDataForBlendProfiles[BlendDataIndex - 1];
+					Context.AppendTask(FAnimNextBlendAddKeyframePerBoneWithScaleTask::Make(SampleBlendDataA, SampleBlendDataB, SampleBlendDataA.GetClampedWeight(), true));
+				}
+
+				// Ensure that rotations are normalized after blending 
+				Context.AppendTask(FAnimNextNormalizeKeyframeRotationsTask());
 			}
+		}
+		else
+		{
+			if (InstanceData->PoseWeights.Num() > 0)
+			{
+				Context.AppendTask(FAnimNextBlendOverwriteKeyframeWithScaleTask::Make(InstanceData->PoseWeights[0]));
+				for (int32 PoseWeightIndex = 1; PoseWeightIndex < InstanceData->PoseWeights.Num(); ++PoseWeightIndex)
+				{
+					Context.AppendTask(FAnimNextBlendAddKeyframeWithScaleTask::Make(InstanceData->PoseWeights[PoseWeightIndex]));
+				}
 
-			// Ensure that rotations are normalized after the blending 
-			Context.AppendTask(FAnimNextNormalizeKeyframeRotationsTask());
+				// Ensure that rotations are normalized after blending 
+				Context.AppendTask(FAnimNextNormalizeKeyframeRotationsTask());
+			}
 		}
-		
-		// Finally we add the additive poses and blend them in
-		if (InstanceData->AdditivePoseWeights.Num() > 0)
+
+		// Finally, we add the additive poses and blend them in
+		if (InstanceData->AdditiveEvalTaskDataForSlot.Num() > 0)
 		{
-			check(InstanceData->AdditivePoseWeights.Num() == InstanceData->AdditiveEvalTaskDataForSlot.Num());
+			const bool bValidAdditiveData = InstanceData->bUsesBlendProfiles ?
+				InstanceData->AdditiveBlendDataForBlendProfiles.Num() == InstanceData->AdditiveEvalTaskDataForSlot.Num() :
+				InstanceData->AdditivePoseWeights.Num() == InstanceData->AdditiveEvalTaskDataForSlot.Num();
 
-			for (int32 iAdditivePose = 0; iAdditivePose < InstanceData->AdditiveEvalTaskDataForSlot.Num(); iAdditivePose++)
+			if (ensureMsgf(bValidAdditiveData, TEXT("Additive montage pose data and blend data does not match! Will skip additives.")))
 			{
-				const FMontageEvalTaskData& TaskData = InstanceData->AdditiveEvalTaskDataForSlot[iAdditivePose];
-				if (TaskData.WeakAnimationSequence.IsValid())
+				for (int32 AdditivePoseIndex = 0; AdditivePoseIndex < InstanceData->AdditiveEvalTaskDataForSlot.Num(); ++AdditivePoseIndex)
 				{
-					FAnimNextAnimSequenceKeyframeTask Task;
-					Task.bInterpolate = true;
-					Task.bLooping = false;
-					Task.bExtractTrajectory = true;	/*Output.AnimInstanceProxy->ShouldExtractRootMotion()*/
-					Task.AnimSequence = TaskData.WeakAnimationSequence.Get();
-					Task.SampleTime = TaskData.AnimPosition;
-					Task.DeltaTimeRecord = FDeltaTimeRecord(DeltaTime);
+					const FMontageEvalTaskData& TaskData = InstanceData->AdditiveEvalTaskDataForSlot[AdditivePoseIndex];
+					if (TaskData.WeakAnimationSequence.IsValid())
+					{
+						FAnimNextAnimSequenceKeyframeTask Task;
+						Task.bInterpolate = true;
+						Task.bLooping = false;
+						Task.bExtractTrajectory = true;	/*Output.AnimInstanceProxy->ShouldExtractRootMotion()*/
+						Task.AnimSequence = TaskData.WeakAnimationSequence.Get();
+						Task.SampleTime = TaskData.AnimPosition;
+						Task.DeltaTimeRecord = FDeltaTimeRecord(DeltaTime);
+
+						Context.AppendTask(Task);
 
-					Context.AppendTask(Task);
+						if (InstanceData->bUsesBlendProfiles)
+						{
+							const FBlendSampleData& AdditiveSampleBlendData = InstanceData->AdditiveBlendDataForBlendProfiles[AdditivePoseIndex];
+							Context.AppendTask(FUAFApplyAdditiveKeyframePerBoneTask::Make(AdditiveSampleBlendData.GetClampedWeight(), AdditiveSampleBlendData, true));
+						}
+						else
+						{
+							Context.AppendTask(FAnimNextApplyAdditiveKeyframeTask::Make(InstanceData->AdditivePoseWeights[AdditivePoseIndex]));
+						}
 
-					Context.AppendTask(FAnimNextApplyAdditiveKeyframeTask::Make(InstanceData->AdditivePoseWeights[iAdditivePose]));
+					}
 				}
 			}
 		}
diff --git a/Engine/Plugins/Experimental/UAF/UAFAnimGraph/Source/UAFAnimGraph/Private/Injection/UAFMontageComponent.cpp b/Engine/Plugins/Experimental/UAF/UAFAnimGraph/Source/UAFAnimGraph/Private/Injection/UAFMontageComponent.cpp
--- a/Engine/Plugins/Experimental/UAF/UAFAnimGraph/Source/UAFAnimGraph/Private/Injection/UAFMontageComponent.cpp
+++ b/Engine/Plugins/Experimental/UAF/UAFAnimGraph/Source/UAFAnimGraph/Private/Injection/UAFMontageComponent.cpp
@@ -36,10 +36,13 @@ void FUAFUpdateMontageTickFunction::ExecuteTick(float DeltaTime, enum ELevelTick
 	if (ensureMsgf(ModuleInstance != nullptr, TEXT("UAF TickMontages: Module Instance was not valid during tick, cannot send notify trait events!")))
 	{
 		// Queue up montage notifies, these will be fired no matter the slot weight. 
-		TSharedPtr<UE::UAF::FNotifyDispatchEvent> NotifyDispatchEvent = MakeTraitEvent<UE::UAF::FNotifyDispatchEvent>();
-		NotifyDispatchEvent->Notifies = WeakAnimInstance->NotifyQueue.AnimNotifies;
-		NotifyDispatchEvent->Weight = 1.0f;
-		ModuleInstance->QueueOutputTraitEvent(NotifyDispatchEvent);
+		if (WeakAnimInstance->NotifyQueue.AnimNotifies.Num() > 0)
+		{
+			TSharedPtr<UE::UAF::FNotifyDispatchEvent> NotifyDispatchEvent = MakeTraitEvent<UE::UAF::FNotifyDispatchEvent>();
+			NotifyDispatchEvent->Notifies = WeakAnimInstance->NotifyQueue.AnimNotifies;
+			NotifyDispatchEvent->Weight = 1.0f;
+			ModuleInstance->QueueOutputTraitEvent(NotifyDispatchEvent);
+		}
 	}
 	
 	// Cache montage animation notifies, these will be fired only if the animation has weight. 
@@ -60,8 +63,10 @@ void FUAFUpdateMontageTickFunction::ExecuteTick(float DeltaTime, enum ELevelTick
 				(
 					MontageInstance->Montage,
 					MontageInstance->GetPosition(),
+					MontageInstance->GetPlayRate(),
+					MontageInstance->GetBlendStartAlpha(),
 					MontageInstance->GetBlend(),
-					MontageInstance->GetPlayRate()
+					MontageInstance->GetActiveBlendProfile()
 				));
 		}
 	}
