UAF: Move accessing of skelmesh component into GT Task because based on init order FAnimNextSkeletalMeshComponentReferenceComponent might not be ready yet

diff --git a/Engine/Plugins/Experimental/UAF/UAFAnimGraph/Source/UAFAnimGraph/Private/Injection/UAFMontageComponent.cpp b/Engine/Plugins/Experimental/UAF/UAFAnimGraph/Source/UAFAnimGraph/Private/Injection/UAFMontageComponent.cpp
--- a/Engine/Plugins/Experimental/UAF/UAFAnimGraph/Source/UAFAnimGraph/Private/Injection/UAFMontageComponent.cpp
+++ b/Engine/Plugins/Experimental/UAF/UAFAnimGraph/Source/UAFAnimGraph/Private/Injection/UAFMontageComponent.cpp
@@ -123,30 +123,15 @@ void FUAFMontageComponent::OnBindToInstance()
 		return;
 	}
 
-	const FAnimNextSkeletalMeshComponentReferenceComponent& ComponentReference = ModuleInstance->GetComponent<FAnimNextSkeletalMeshComponentReferenceComponent>();
-	USkeletalMeshComponent* SkeletalMeshComponent = ComponentReference.GetComponent();
-
-	if (!SkeletalMeshComponent || SkeletalMeshComponent->bEnableAnimation)
-	{
-		return;
-	}
-
 	UAnimNextComponent* UAFComponent = Cast<UAnimNextComponent>(ModuleInstance->GetObject());
 	if (!UAFComponent)
 	{
 		return;
 	}
 
-	MontageTickFunction.bRunOnAnyThread = false;
-	MontageTickFunction.bCanEverTick = true;
-	MontageTickFunction.bStartWithTickEnabled = true;
-	MontageTickFunction.MontageComponent = this;
-	MontageTickFunction.Target = UAFComponent;
-
 	FAnimNextModuleInstance::RunTaskOnGameThread(
-		[WeakSkelMeshComponent = TWeakObjectPtr<USkeletalMeshComponent>(SkeletalMeshComponent), WeakUAFComponent = TWeakObjectPtr<UAnimNextComponent>(UAFComponent), this]()
+		[WeakUAFComponent = TWeakObjectPtr<UAnimNextComponent>(UAFComponent), MontageComponent = this]()
 		{
-
 			UAnimNextComponent* UAFComponent = WeakUAFComponent.Get();
 			if (!UAFComponent)
 			{
@@ -159,8 +144,8 @@ void FUAFMontageComponent::OnBindToInstance()
 				return;
 			}
 
-			USkeletalMeshComponent* SkelMeshComponent = WeakSkelMeshComponent.Get();
-			if (!SkelMeshComponent)
+			USkeletalMeshComponent* SkelMeshComponent = UAFActor->FindComponentByClass<USkeletalMeshComponent>();
+			if (!SkelMeshComponent || SkelMeshComponent->bEnableAnimation)
 			{
 				return;
 			}
@@ -181,14 +166,19 @@ void FUAFMontageComponent::OnBindToInstance()
 			UsedAnimInstance->InitializeMontageOnly();
 
 			// Setup and register tick function
-			MontageTickFunction.WeakAnimInstance = UsedAnimInstance;
+			MontageComponent->MontageTickFunction.bRunOnAnyThread = false;
+			MontageComponent->MontageTickFunction.bCanEverTick = true;
+			MontageComponent->MontageTickFunction.bStartWithTickEnabled = true;
+			MontageComponent->MontageTickFunction.MontageComponent = MontageComponent;
+			MontageComponent->MontageTickFunction.Target = UAFComponent;
+			MontageComponent->MontageTickFunction.WeakAnimInstance = UsedAnimInstance;
 
 			if (UWorld* World = UAFComponent->GetWorld())
 			{
 				ULevel* Level = World->PersistentLevel;
-				MontageTickFunction.RegisterTickFunction(Level);
+				MontageComponent->MontageTickFunction.RegisterTickFunction(Level);
 
-				UAFComponent->AddPrerequisite(UAFComponent, MontageTickFunction, TickDependecyEventName);
+				UAFComponent->AddPrerequisite(UAFComponent, MontageComponent->MontageTickFunction, TickDependecyEventName);
 			}
 		});
 }
