UAF: Added CopyBonesComponentSpace trait

diff --git a/Engine/Plugins/Experimental/UAF/UAFAnimGraph/Source/UAFAnimGraph/Private/AnimNextAnimGraphModule.cpp b/Engine/Plugins/Experimental/UAF/UAFAnimGraph/Source/UAFAnimGraph/Private/AnimNextAnimGraphModule.cpp
--- a/Engine/Plugins/Experimental/UAF/UAFAnimGraph/Source/UAFAnimGraph/Private/AnimNextAnimGraphModule.cpp
+++ b/Engine/Plugins/Experimental/UAF/UAFAnimGraph/Source/UAFAnimGraph/Private/AnimNextAnimGraphModule.cpp
@@ -22,6 +22,7 @@
 #include "Traits/BlendSpacePlayerTraitData.h"
 #include "Traits/NotifyDispatcherTraitData.h"
 #include "Traits/SequencePlayer.h"
+#include "Traits/BoneMapping.h"
 #include "Factory/AnimNextFactoryParams.h"
 #include "Factory/SystemFactory.h"
 #include "Factory/UAFSystemFactoryParams.h"
@@ -73,6 +74,7 @@ void FAnimNextAnimGraphModule::StartupModule()
 		FRigVMTrait_AnimNextPublicVariables::StaticStruct(),
 		FAnimNextInjectionBlendSettings::StaticStruct(),
 		FModifyCurveParameters::StaticStruct(),
+		FAnimNextBoneMapping::StaticStruct(),
 	};
 
 	RigVMRegistry.RegisterStructTypes(AllowedStructTypes);
diff --git a/Engine/Plugins/Experimental/UAF/UAFAnimGraph/Source/UAFAnimGraph/Private/EvaluationVM/Tasks/CopyBones.cpp b/Engine/Plugins/Experimental/UAF/UAFAnimGraph/Source/UAFAnimGraph/Private/EvaluationVM/Tasks/CopyBones.cpp
new file mode 100644
--- /dev/null
+++ b/Engine/Plugins/Experimental/UAF/UAFAnimGraph/Source/UAFAnimGraph/Private/EvaluationVM/Tasks/CopyBones.cpp
@@ -0,0 +1,82 @@
+// Copyright Epic Games, Inc. All Rights Reserved.
+
+#include "EvaluationVM/Tasks/CopyBones.h"
+
+#include "EvaluationVM/EvaluationVM.h"
+#include "EvaluationVM/KeyframeState.h"
+
+#include UE_INLINE_GENERATED_CPP_BY_NAME(CopyBones)
+
+FAnimNextCopyBonesComponentSpaceTask FAnimNextCopyBonesComponentSpaceTask::Make(const TArray<FAnimNextBoneMapping>& Mapping)
+{
+	FAnimNextCopyBonesComponentSpaceTask Task;
+	Task.Mapping = Mapping;
+	return Task;
+}
+
+void FAnimNextCopyBonesComponentSpaceTask::Execute(UE::UAF::FEvaluationVM& VM) const
+{
+	using namespace UE::UAF;
+
+	if (VM.GetContext().IsValid())
+	{
+		// TODO: Implement with new attribute runtime
+		return;
+	}
+
+	if (EnumHasAnyFlags(VM.GetFlags(), EEvaluationFlags::Bones))
+	{
+		if (const TUniquePtr<FKeyframeState>* Keyframe = VM.PeekValue<TUniquePtr<FKeyframeState>>(KEYFRAME_STACK_NAME, 0))
+		{
+			// Get the number of bones to copy
+
+			const int32 MappingNum = Mapping.Num();
+
+			// Loop over each pair of bones to copy
+
+			for (int32 MappingIdx = 0; MappingIdx < MappingNum; MappingIdx++)
+			{
+				// Find LOD indices for both bones
+
+				const FBoneIndexType SourceBoneIdx = Keyframe->Get()->Pose.GetRefPose().FindLODBoneIndexFromBoneName(Mapping[MappingIdx].SourceBone);
+				const FBoneIndexType DestinationBoneIdx = Keyframe->Get()->Pose.GetRefPose().FindLODBoneIndexFromBoneName(Mapping[MappingIdx].TargetBone);
+
+				// If we can't find them or they are not in the current LOD then skip
+
+				if (SourceBoneIdx == FBoneIndexType(INDEX_NONE) ||
+					DestinationBoneIdx == FBoneIndexType(INDEX_NONE) ||
+					SourceBoneIdx >= Keyframe->Get()->Pose.GetNumBones() ||
+					DestinationBoneIdx >= Keyframe->Get()->Pose.GetNumBones())
+				{
+					continue;
+				}
+
+				// Compute the Source Component Space Transform
+
+				FTransform SourceComponentTransform = Keyframe->Get()->Pose.LocalTransforms[SourceBoneIdx];
+				int32 SourceParentIdx = Keyframe->Get()->Pose.GetRefPose().GetLODParentBoneIndex(Keyframe->Get()->Pose.LODLevel, SourceBoneIdx);
+				while (SourceParentIdx != INDEX_NONE)
+				{
+					FTransform ParentLocalTransform = Keyframe->Get()->Pose.LocalTransforms[SourceParentIdx];
+					SourceComponentTransform = SourceComponentTransform * ParentLocalTransform;
+					SourceParentIdx = Keyframe->Get()->Pose.GetRefPose().GetLODParentBoneIndex(Keyframe->Get()->Pose.LODLevel, SourceParentIdx);
+				}
+
+				// Compute the Destination Parent Component Space Transform
+
+				FTransform DestinationComponentParentTransform = FTransform::Identity;
+				int32 DestinationParentIdx = Keyframe->Get()->Pose.GetRefPose().GetLODParentBoneIndex(Keyframe->Get()->Pose.LODLevel, DestinationBoneIdx);
+				while (DestinationParentIdx != INDEX_NONE)
+				{
+					FTransform ParentLocalTransform = Keyframe->Get()->Pose.LocalTransforms[DestinationParentIdx];
+					DestinationComponentParentTransform = DestinationComponentParentTransform * ParentLocalTransform;
+					DestinationParentIdx = Keyframe->Get()->Pose.GetRefPose().GetLODParentBoneIndex(Keyframe->Get()->Pose.LODLevel, DestinationParentIdx);
+				}
+
+				// Set the Destination Bone Transform
+
+				Keyframe->Get()->Pose.LocalTransforms[DestinationBoneIdx] = SourceComponentTransform.GetRelativeTransform(DestinationComponentParentTransform);
+			}
+		}
+	}
+}
diff --git a/Engine/Plugins/Experimental/UAF/UAFAnimGraph/Source/UAFAnimGraph/Private/EvaluationVM/Tasks/CopyBones.h b/Engine/Plugins/Experimental/UAF/UAFAnimGraph/Source/UAFAnimGraph/Private/EvaluationVM/Tasks/CopyBones.h
new file mode 100644
--- /dev/null
+++ b/Engine/Plugins/Experimental/UAF/UAFAnimGraph/Source/UAFAnimGraph/Private/EvaluationVM/Tasks/CopyBones.h
@@ -0,0 +1,34 @@
+// Copyright Epic Games, Inc. All Rights Reserved.
+
+#pragma once
+
+#include "CoreMinimal.h"
+#include "EvaluationVM/EvaluationTask.h"
+#include "Traits/BoneMapping.h"
+
+#include "CopyBones.generated.h"
+
+#define UE_API UAFANIMGRAPH_API
+
+/**
+ * FAnimNextCopyBonesComponentSpaceTask
+ * 
+ * Copies bones in component space
+ */
+USTRUCT()
+struct FAnimNextCopyBonesComponentSpaceTask : public FAnimNextEvaluationTask
+{
+	GENERATED_BODY()
+
+	DECLARE_ANIM_EVALUATION_TASK(FAnimNextCopyBonesComponentSpaceTask)
+
+	// Make a FAnimNextCopyBonesComponentSpaceTask from a source to destination bone mapping
+	static UE_API FAnimNextCopyBonesComponentSpaceTask Make(const TArray<FAnimNextBoneMapping>& Mapping);
+
+	// Task entry point
+	UE_API virtual void Execute(UE::UAF::FEvaluationVM& VM) const override;
+
+	TArray<FAnimNextBoneMapping> Mapping;
+};
+
+#undef UE_API
diff --git a/Engine/Plugins/Experimental/UAF/UAFAnimGraph/Source/UAFAnimGraph/Private/Traits/CopyBonesComponentSpace.cpp b/Engine/Plugins/Experimental/UAF/UAFAnimGraph/Source/UAFAnimGraph/Private/Traits/CopyBonesComponentSpace.cpp
new file mode 100644
--- /dev/null
+++ b/Engine/Plugins/Experimental/UAF/UAFAnimGraph/Source/UAFAnimGraph/Private/Traits/CopyBonesComponentSpace.cpp
@@ -0,0 +1,28 @@
+// Copyright Epic Games, Inc. All Rights Reserved.
+
+#include "Traits/CopyBonesComponentSpace.h"
+
+#include "EvaluationVM/Tasks/CopyBones.h"
+
+#include UE_INLINE_GENERATED_CPP_BY_NAME(CopyBonesComponentSpace)
+
+namespace UE::UAF
+{
+	AUTO_REGISTER_ANIM_TRAIT(FCopyBonesComponentSpaceTrait)
+
+		// Trait implementation boilerplate
+#define TRAIT_INTERFACE_ENUMERATOR(GeneratorMacro) \
+		GeneratorMacro(IEvaluate) \
+
+		GENERATE_ANIM_TRAIT_IMPLEMENTATION(FCopyBonesComponentSpaceTrait, TRAIT_INTERFACE_ENUMERATOR, NULL_ANIM_TRAIT_INTERFACE_ENUMERATOR, NULL_ANIM_TRAIT_EVENT_ENUMERATOR)
+#undef TRAIT_INTERFACE_ENUMERATOR
+
+	void FCopyBonesComponentSpaceTrait::PostEvaluate(FEvaluateTraversalContext& Context, const TTraitBinding<IEvaluate>& Binding) const
+	{
+		IEvaluate::PostEvaluate(Context, Binding);
+
+		const FSharedData* SharedData = Binding.GetSharedData<FSharedData>();
+        
+        Context.AppendTask(FAnimNextCopyBonesComponentSpaceTask::Make(SharedData->BoneMapping));
+	}
+}
diff --git a/Engine/Plugins/Experimental/UAF/UAFAnimGraph/Source/UAFAnimGraph/Private/Traits/CopyBonesComponentSpace.h b/Engine/Plugins/Experimental/UAF/UAFAnimGraph/Source/UAFAnimGraph/Private/Traits/CopyBonesComponentSpace.h
new file mode 100644
--- /dev/null
+++ b/Engine/Plugins/Experimental/UAF/UAFAnimGraph/Source/UAFAnimGraph/Private/Traits/CopyBonesComponentSpace.h
@@ -0,0 +1,42 @@
+// Copyright Epic Games, Inc. All Rights Reserved.
+
+#pragma once
+
+#include "CoreMinimal.h"
+
+#include "TraitCore/Trait.h"
+#include "TraitInterfaces/IEvaluate.h"
+#include "Traits/BoneMapping.h"
+
+#include "CopyBonesComponentSpace.generated.h"
+
+/** A trait that allows you to copy bone transforms in component space. */
+USTRUCT(meta = (DisplayName = "Copy Bones in Component Space", ShowTooltip=true))
+struct FAnimNextCopyBonesComponentSpaceTraitSharedData : public FAnimNextTraitSharedData
+{
+	GENERATED_BODY()
+
+	// Bones to copy transforms between
+	UPROPERTY(EditAnywhere, Category = Settings, meta = (Inline))
+	TArray<FAnimNextBoneMapping> BoneMapping;
+};
+
+namespace UE::UAF
+{
+	/**
+	 * FCopyBonesComponentSpaceTrait
+	 *
+	 * A trait that allows bone transforms to be copied in component space.
+	 */
+	struct FCopyBonesComponentSpaceTrait : FAdditiveTrait, IEvaluate
+	{
+		DECLARE_ANIM_TRAIT(FCopyBonesComponentSpaceTrait, FAdditiveTrait)
+
+		using FSharedData = FAnimNextCopyBonesComponentSpaceTraitSharedData;
+
+		struct FInstanceData : FTrait::FInstanceData {};
+
+		// IEvaluate impl
+		virtual void PostEvaluate(FEvaluateTraversalContext& Context, const TTraitBinding<IEvaluate>& Binding) const override;
+	};
+}
diff --git a/Engine/Plugins/Experimental/UAF/UAFAnimGraph/Source/UAFAnimGraph/Public/Traits/BoneMapping.h b/Engine/Plugins/Experimental/UAF/UAFAnimGraph/Source/UAFAnimGraph/Public/Traits/BoneMapping.h
new file mode 100644
--- /dev/null
+++ b/Engine/Plugins/Experimental/UAF/UAFAnimGraph/Source/UAFAnimGraph/Public/Traits/BoneMapping.h
@@ -0,0 +1,21 @@
+// Copyright Epic Games, Inc. All Rights Reserved.
+
+#pragma once
+
+#include "CoreMinimal.h"
+
+#include "BoneMapping.generated.h"
+
+USTRUCT(BlueprintType)
+struct FAnimNextBoneMapping
+{
+	GENERATED_BODY()
+
+	/** Bone to map from  */
+	UPROPERTY(EditAnywhere, Category = "Default")
+	FName SourceBone;
+
+	/** Bone to map to */
+	UPROPERTY(EditAnywhere, Category = "Default")
+	FName TargetBone;
+};
\ No newline at end of file
