UAF: Change looping flag to an enum

diff --git a/Engine/Plugins/Experimental/UAF/UAFAnimGraph/Source/UAFAnimGraph/Private/AnimNextAnimGraphModule.cpp b/Engine/Plugins/Experimental/UAF/UAFAnimGraph/Source/UAFAnimGraph/Private/AnimNextAnimGraphModule.cpp
--- a/Engine/Plugins/Experimental/UAF/UAFAnimGraph/Source/UAFAnimGraph/Private/AnimNextAnimGraphModule.cpp
+++ b/Engine/Plugins/Experimental/UAF/UAFAnimGraph/Source/UAFAnimGraph/Private/AnimNextAnimGraphModule.cpp
@@ -169,10 +169,6 @@ void FAnimNextAnimGraphModule::RegisterGraphFactories()
 				ensure(InContext.AccessStruct<FSequencePlayerData>([&InContext](FSequencePlayerData& InStruct)
 				{
 					InStruct.AnimSequence = CastChecked<UAnimSequence>(InContext.GetObject());
-					if (InStruct.AnimSequence)
-					{
-						InStruct.bLoop = InStruct.AnimSequence->bLoop;
-					}
 				}));
 			});
 	}
@@ -190,10 +186,6 @@ void FAnimNextAnimGraphModule::RegisterGraphFactories()
 				ensure(InContext.AccessStruct<FBlendSpacePlayerData>([&InContext](FBlendSpacePlayerData& InStruct)
 				{
 					InStruct.BlendSpace = CastChecked<UBlendSpace>(InContext.GetObject());
-					if (InStruct.BlendSpace)
-					{
-						InStruct.bLoop = InStruct.BlendSpace->bLoop;
-					}
 				}));
 			});
 	}
@@ -239,10 +231,6 @@ void FAnimNextAnimGraphModule::RegisterSystemFactories()
 				ensure(InContext.AccessVariablesStruct<FSequencePlayerData>([&InContext](FSequencePlayerData& InStruct)
 				{
 					InStruct.AnimSequence = CastChecked<UAnimSequence>(InContext.GetObject());
-					if (InStruct.AnimSequence)
-					{
-						InStruct.bLoop = InStruct.AnimSequence->bLoop;
-					}
 				}));
 			});
 	}
@@ -267,10 +255,6 @@ void FAnimNextAnimGraphModule::RegisterSystemFactories()
 				ensure(InContext.AccessVariablesStruct<FBlendSpacePlayerData>([&InContext](FBlendSpacePlayerData& InStruct)
 				{
 					InStruct.BlendSpace = CastChecked<UBlendSpace>(InContext.GetObject());
-					if (InStruct.BlendSpace)
-					{
-						InStruct.bLoop = InStruct.BlendSpace->bLoop;
-					}
 				}));
 			});
 	}
diff --git a/Engine/Plugins/Experimental/UAF/UAFAnimGraph/Source/UAFAnimGraph/Private/Injection/InjectionUtils.cpp b/Engine/Plugins/Experimental/UAF/UAFAnimGraph/Source/UAFAnimGraph/Private/Injection/InjectionUtils.cpp
--- a/Engine/Plugins/Experimental/UAF/UAFAnimGraph/Source/UAFAnimGraph/Private/Injection/InjectionUtils.cpp
+++ b/Engine/Plugins/Experimental/UAF/UAFAnimGraph/Source/UAFAnimGraph/Private/Injection/InjectionUtils.cpp
@@ -1,4 +1,4 @@
-﻿// Copyright Epic Games, Inc. All Rights Reserved.
+// Copyright Epic Games, Inc. All Rights Reserved.
 
 #include "Injection/InjectionUtils.h"
 #include "GraphInterfaces/AnimNextNativeDataInterface_AnimSequencePlayer.h"
@@ -107,19 +107,7 @@ namespace UE::UAF
 						.AnimSequence = InAnimSequence,
 						.PlayRate = InArgs.PlayRate,
 						.StartPosition = InArgs.StartPosition,
-						.bLoop = [&InArgs, &InAnimSequence]()
-						{
-							switch (InArgs.LoopMode)
-							{
-							case ELoopMode::Auto:
-								return InAnimSequence->bLoop;
-							case ELoopMode::ForceLoop:
-								return true;
-							case ELoopMode::ForceNonLoop:
-								return false;
-							}
-							return false;
-						}()
+						.LoopMode = InArgs.LoopMode
 					}));
 		}
 		else
diff --git a/Engine/Plugins/Experimental/UAF/UAFAnimGraph/Source/UAFAnimGraph/Private/Traits/BlendSpacePlayer.cpp b/Engine/Plugins/Experimental/UAF/UAFAnimGraph/Source/UAFAnimGraph/Private/Traits/BlendSpacePlayer.cpp
--- a/Engine/Plugins/Experimental/UAF/UAFAnimGraph/Source/UAFAnimGraph/Private/Traits/BlendSpacePlayer.cpp
+++ b/Engine/Plugins/Experimental/UAF/UAFAnimGraph/Source/UAFAnimGraph/Private/Traits/BlendSpacePlayer.cpp
@@ -103,6 +103,7 @@ namespace UE::UAF
 				for (int32 CurSampleIndex = 0; CurSampleIndex < BlendSamples.Num(); CurSampleIndex++)
 				{
 					const UAnimSequence* CurSequence = BlendSamples[CurSampleIndex].Animation;
+					const EAnimAssetLoopMode LoopMode = SharedData->GetLoopMode(Binding);
 
 					FAnimNextFactoryParams FactoryParams = FAnimGraphFactory::GetDefaultParamsForClass(CurSequence->GetClass());
 					FactoryParams.PushPublicTrait(
@@ -111,7 +112,19 @@ namespace UE::UAF
 							.AnimSequence = CurSequence,
 							// Assume start position is normalized.
 							.StartPosition = SharedData->GetStartPosition(Binding) * BlendSamples[CurSampleIndex].GetSamplePlayLength(),
-							.bLoop = SharedData->GetbLoop(Binding)
+							// Blendspace has authority over loopness, since non-looping and looping assets can't sync together.
+							.LoopMode = [&InstanceData, LoopMode]()
+							{
+								switch (LoopMode)
+								{
+									case EAnimAssetLoopMode::Auto:
+										return InstanceData->BlendSpace->bLoop ? EAnimAssetLoopMode::ForceLoop : EAnimAssetLoopMode::ForceNonLoop;
+									case EAnimAssetLoopMode::ForceLoop:
+									case EAnimAssetLoopMode::ForceNonLoop:
+									default:
+										return LoopMode;
+								}
+							}()
 						}));
 
 					const UAnimNextAnimationGraph* AnimationGraph = IGraphFactory::GetOrBuildGraph(Context, Binding, CurSequence, FactoryParams);
diff --git a/Engine/Plugins/Experimental/UAF/UAFAnimGraph/Source/UAFAnimGraph/Private/Traits/SequencePlayer.cpp b/Engine/Plugins/Experimental/UAF/UAFAnimGraph/Source/UAFAnimGraph/Private/Traits/SequencePlayer.cpp
--- a/Engine/Plugins/Experimental/UAF/UAFAnimGraph/Source/UAFAnimGraph/Private/Traits/SequencePlayer.cpp
+++ b/Engine/Plugins/Experimental/UAF/UAFAnimGraph/Source/UAFAnimGraph/Private/Traits/SequencePlayer.cpp
@@ -24,6 +24,21 @@ namespace UE::UAF
 	GENERATE_ANIM_TRAIT_IMPLEMENTATION(FSequencePlayerTrait, TRAIT_INTERFACE_ENUMERATOR, NULL_ANIM_TRAIT_INTERFACE_ENUMERATOR, NULL_ANIM_TRAIT_EVENT_ENUMERATOR)
 	#undef TRAIT_INTERFACE_ENUMERATOR
 
+	static bool IsSequenceLoopingHelper(const UAnimSequence* InAnimSequence, EAnimAssetLoopMode InLoopMode)
+	{
+		switch (InLoopMode)
+		{
+			case EAnimAssetLoopMode::Auto:
+				return InAnimSequence ? InAnimSequence->bLoop : false;
+			case EAnimAssetLoopMode::ForceLoop:
+				return true;
+			case EAnimAssetLoopMode::ForceNonLoop:
+				return false;
+			default:
+				return false;
+		}
+	}
+
 	void FSequencePlayerTrait::PreEvaluate(FEvaluateTraversalContext& Context, const TTraitBinding<IEvaluate>& Binding) const
 	{
 		const FSharedData* SharedData = Binding.GetSharedData<FSharedData>();
@@ -36,7 +51,7 @@ namespace UE::UAF
 		Task.SampleTime = InstanceData->InternalTimeAccumulator;
 		Task.bInterpolate = bInterpolate;
 		Task.DeltaTimeRecord = InstanceData->DeltaTimeRecord;
-		Task.bLooping = SharedData->GetbLoop(Binding);
+		Task.bLooping = IsSequenceLoopingHelper(InstanceData->AnimSequence, SharedData->GetLoopMode(Binding));
 		Task.bExtractTrajectory = true;	/*Output.AnimInstanceProxy->ShouldExtractRootMotion()*/
 
 		Context.AppendTask(Task);
@@ -66,7 +81,7 @@ namespace UE::UAF
 				// Account for play rate
 				DeltaTime *= PlayRate;
 
-				return AnimSequence->ExtractRootMotion(FAnimExtractContext(static_cast<double>(StartTime), true, FDeltaTimeRecord(DeltaTime), bAllowLooping && bIsLooping && AnimSequence->bLoop));
+				return AnimSequence->ExtractRootMotion(FAnimExtractContext(static_cast<double>(StartTime), true, FDeltaTimeRecord(DeltaTime), bAllowLooping && bIsLooping));
 			};
 
 			return FOnExtractRootMotionAttribute::CreateLambda(ExtractRootMotionAttribute);
@@ -100,7 +115,7 @@ namespace UE::UAF
 				InstanceData->InternalTimeAccumulator,
 				AnimSeq->GetPlayLength(),
 				SharedData->GetPlayRate(Binding) * AnimSeq->RateScale,
-				SharedData->GetbLoop(Binding)).WithDebugName(AnimSeq->GetFName());
+				IsSequenceLoopingHelper(AnimSeq, SharedData->GetLoopMode(Binding))).WithDebugName(AnimSeq->GetFName());
 		}
 
 		return FTimelineState();
diff --git a/Engine/Plugins/Experimental/UAF/UAFAnimGraph/Source/UAFAnimGraph/Public/Factory/AnimAssetParams.h b/Engine/Plugins/Experimental/UAF/UAFAnimGraph/Source/UAFAnimGraph/Public/Factory/AnimAssetParams.h
new file mode 100644
--- /dev/null
+++ b/Engine/Plugins/Experimental/UAF/UAFAnimGraph/Source/UAFAnimGraph/Public/Factory/AnimAssetParams.h
@@ -0,0 +1,21 @@
+// Copyright Epic Games, Inc. All Rights Reserved.
+
+#pragma once
+
+#include "AnimAssetParams.generated.h"
+
+namespace UE::UAF
+{
+	UENUM(BlueprintType)
+	enum class EAnimAssetLoopMode : uint8
+	{
+		/** Asset will use the looping value defined in the asset. */
+		Auto,
+
+		/** Asset will loop regardless of the value defined in the asset. */
+		ForceLoop,
+
+		/** Asset will NOT loop regardless of the value defined in the asset. */
+		ForceNonLoop,
+	};
+}
\ No newline at end of file
diff --git a/Engine/Plugins/Experimental/UAF/UAFAnimGraph/Source/UAFAnimGraph/Public/Injection/InjectionUtils.h b/Engine/Plugins/Experimental/UAF/UAFAnimGraph/Source/UAFAnimGraph/Public/Injection/InjectionUtils.h
--- a/Engine/Plugins/Experimental/UAF/UAFAnimGraph/Source/UAFAnimGraph/Public/Injection/InjectionUtils.h
+++ b/Engine/Plugins/Experimental/UAF/UAFAnimGraph/Source/UAFAnimGraph/Public/Injection/InjectionUtils.h
@@ -5,6 +5,7 @@
 #include "CoreMinimal.h"
 #include "InjectionRequest.h"
 #include "Factory/AnimNextFactoryParams.h"
+#include "Factory/AnimAssetParams.h"
 #include "Injection/InjectionSite.h"
 
 struct FUAFAssetInstance;
@@ -17,28 +18,13 @@ namespace UE::UAF
 
 namespace UE::UAF
 {
-	// How to loop an animation when playing back
-	enum class ELoopMode : uint8
-	{
-		// Use looping flag on the animation
-		Auto,
-
-		// Force looping
-		ForceLoop,
-
-		// Force non-looping
-		ForceNonLoop,
-
-		// TODO: Specific fractional loop count
-	};
-
 	// Optional arguments to PlayAnim
 	struct FPlayAnimArgs
 	{
 		FAnimNextFactoryParams FactoryParams;
 		float PlayRate = 1.0f;
 		float StartPosition = 0.0f;
-		ELoopMode LoopMode = ELoopMode::Auto;
+		EAnimAssetLoopMode LoopMode = EAnimAssetLoopMode::Auto;
 		EAnimNextInjectionLifetimeType LifetimeType = EAnimNextInjectionLifetimeType::Auto;
 	};
 
diff --git a/Engine/Plugins/Experimental/UAF/UAFAnimGraph/Source/UAFAnimGraph/Public/Traits/BlendSpacePlayerTraitData.h b/Engine/Plugins/Experimental/UAF/UAFAnimGraph/Source/UAFAnimGraph/Public/Traits/BlendSpacePlayerTraitData.h
--- a/Engine/Plugins/Experimental/UAF/UAFAnimGraph/Source/UAFAnimGraph/Public/Traits/BlendSpacePlayerTraitData.h
+++ b/Engine/Plugins/Experimental/UAF/UAFAnimGraph/Source/UAFAnimGraph/Public/Traits/BlendSpacePlayerTraitData.h
@@ -1,4 +1,4 @@
-﻿// Copyright Epic Games, Inc. All Rights Reserved.
+// Copyright Epic Games, Inc. All Rights Reserved.
 
 #pragma once
 
@@ -6,6 +6,7 @@
 #include "TraitCore/Trait.h"
 #include "Animation/BlendSpace.h"
 #include "TraitCore/TraitBinding.h"
+#include "Factory/AnimAssetParams.h"
 
 #include "BlendSpacePlayerTraitData.generated.h"
 
@@ -35,9 +36,9 @@ struct FAnimNextBlendSpacePlayerTraitSharedData : public FAnimNextTraitSharedDat
 	UPROPERTY(EditAnywhere, Category = "Blend Space Player")
 	float StartPosition = 0.0f;
 
-	/** Whether to loop the animation */
+	/** Whether to use the asset's default looping behavior or force an override. */
 	UPROPERTY(EditAnywhere, Category = "Blend Space Player")
-	bool bLoop = false;
+	UE::UAF::EAnimAssetLoopMode LoopMode = UE::UAF::EAnimAssetLoopMode::Auto;
 
 	// Latent pin support boilerplate
 	#define TRAIT_LATENT_PROPERTIES_ENUMERATOR(GeneratorMacro) \
@@ -46,7 +47,7 @@ struct FAnimNextBlendSpacePlayerTraitSharedData : public FAnimNextTraitSharedDat
 		GeneratorMacro(YAxisSamplePoint) \
 		GeneratorMacro(PlayRate) \
 		GeneratorMacro(StartPosition) \
-		GeneratorMacro(bLoop) \
+		GeneratorMacro(LoopMode) \
 
 	GENERATE_TRAIT_LATENT_PROPERTIES(FAnimNextBlendSpacePlayerTraitSharedData, TRAIT_LATENT_PROPERTIES_ENUMERATOR)
 	#undef TRAIT_LATENT_PROPERTIES_ENUMERATOR
diff --git a/Engine/Plugins/Experimental/UAF/UAFAnimGraph/Source/UAFAnimGraph/Public/Traits/SequencePlayerTraitData.h b/Engine/Plugins/Experimental/UAF/UAFAnimGraph/Source/UAFAnimGraph/Public/Traits/SequencePlayerTraitData.h
--- a/Engine/Plugins/Experimental/UAF/UAFAnimGraph/Source/UAFAnimGraph/Public/Traits/SequencePlayerTraitData.h
+++ b/Engine/Plugins/Experimental/UAF/UAFAnimGraph/Source/UAFAnimGraph/Public/Traits/SequencePlayerTraitData.h
@@ -5,6 +5,7 @@
 #include "TraitCore/TraitBinding.h"
 #include "TraitCore/TraitSharedData.h"
 #include "Animation/AnimSequence.h"
+#include "Factory/AnimAssetParams.h"
 #include "SequencePlayerTraitData.generated.h"
 
 /** A trait that can play an animation sequence. */
@@ -25,16 +26,16 @@ struct FAnimNextSequencePlayerTraitSharedData : public FAnimNextTraitSharedData
 	UPROPERTY(EditAnywhere, Category = "Sequence Player", meta = (OnBecomeRelevant))
 	float StartPosition = 0.0f;
 
-	/** Whether or not this sequence playback will loop. */
+	/** Whether to use the asset's default looping behavior or force an override. */
 	UPROPERTY(EditAnywhere, Category = "Sequence Player", meta = (OnBecomeRelevant))
-	bool bLoop = false;
+	UE::UAF::EAnimAssetLoopMode LoopMode = UE::UAF::EAnimAssetLoopMode::Auto;
 
 	// Latent pin support boilerplate
 	#define TRAIT_LATENT_PROPERTIES_ENUMERATOR(GeneratorMacro) \
 		GeneratorMacro(AnimSequence) \
 		GeneratorMacro(PlayRate) \
 		GeneratorMacro(StartPosition) \
-		GeneratorMacro(bLoop) \
+		GeneratorMacro(LoopMode) \
 
 	GENERATE_TRAIT_LATENT_PROPERTIES(FAnimNextSequencePlayerTraitSharedData, TRAIT_LATENT_PROPERTIES_ENUMERATOR)
 	#undef TRAIT_LATENT_PROPERTIES_ENUMERATOR
diff --git a/Engine/Plugins/Experimental/UAF/UAFAnimGraph/Source/UAFAnimGraph/UAFAnimGraph.Build.cs b/Engine/Plugins/Experimental/UAF/UAFAnimGraph/Source/UAFAnimGraph/UAFAnimGraph.Build.cs
--- a/Engine/Plugins/Experimental/UAF/UAFAnimGraph/Source/UAFAnimGraph/UAFAnimGraph.Build.cs
+++ b/Engine/Plugins/Experimental/UAF/UAFAnimGraph/Source/UAFAnimGraph/UAFAnimGraph.Build.cs
@@ -38,6 +38,8 @@ namespace UnrealBuildTool.Rules
 					}
 				);
 			}
+
+			bAllowUETypesInNamespaces = true;
 		}
 	}
 }
\ No newline at end of file
diff --git a/Engine/Plugins/Experimental/UAF/UAFPoseSearch/Source/UAFPoseSearch/Private/MotionMatchingTrait.cpp b/Engine/Plugins/Experimental/UAF/UAFPoseSearch/Source/UAFPoseSearch/Private/MotionMatchingTrait.cpp
--- a/Engine/Plugins/Experimental/UAF/UAFPoseSearch/Source/UAFPoseSearch/Private/MotionMatchingTrait.cpp
+++ b/Engine/Plugins/Experimental/UAF/UAFPoseSearch/Source/UAFPoseSearch/Private/MotionMatchingTrait.cpp
@@ -335,7 +335,7 @@ PRAGMA_ENABLE_DEPRECATION_WARNINGS
 
 						BlendSpacePlayer.PlayRate = DesiredPlayRate;
 						BlendSpacePlayer.StartPosition = MotionMatchingState.SearchResult.SelectedTime;
-						BlendSpacePlayer.bLoop = MotionMatchingState.SearchResult.bLoop;
+						BlendSpacePlayer.LoopMode = MotionMatchingState.SearchResult.bLoop ? EAnimAssetLoopMode::ForceLoop : EAnimAssetLoopMode::ForceNonLoop;
 						FactoryParams.PushPublicTrait(BlendSpacePlayer);
 					}
 					else
@@ -344,7 +344,7 @@ PRAGMA_ENABLE_DEPRECATION_WARNINGS
 						AnimSequencePlayer.AnimSequence = Cast<UAnimSequence>(AnimationAsset);
 						AnimSequencePlayer.PlayRate = DesiredPlayRate;
 						AnimSequencePlayer.StartPosition = MotionMatchingState.SearchResult.SelectedTime;
-						AnimSequencePlayer.bLoop = MotionMatchingState.SearchResult.bLoop;
+						AnimSequencePlayer.LoopMode = MotionMatchingState.SearchResult.bLoop ? EAnimAssetLoopMode::ForceLoop : EAnimAssetLoopMode::ForceNonLoop;
 						FactoryParams.PushPublicTrait(AnimSequencePlayer);
 					}
 
diff --git a/Engine/Plugins/Experimental/UAF/UAFPoseSearch/Source/UAFPoseSearch/Private/PoseSearchResultEmulatorTrait.cpp b/Engine/Plugins/Experimental/UAF/UAFPoseSearch/Source/UAFPoseSearch/Private/PoseSearchResultEmulatorTrait.cpp
--- a/Engine/Plugins/Experimental/UAF/UAFPoseSearch/Source/UAFPoseSearch/Private/PoseSearchResultEmulatorTrait.cpp
+++ b/Engine/Plugins/Experimental/UAF/UAFPoseSearch/Source/UAFPoseSearch/Private/PoseSearchResultEmulatorTrait.cpp
@@ -118,7 +118,7 @@ namespace UE::UAF
 					BlendSpacePlayer.YAxisSamplePoint = YAxisSamplePoint;
 					BlendSpacePlayer.PlayRate = WantedPlayRate;
 					BlendSpacePlayer.StartPosition =  SelectedTime;
-					BlendSpacePlayer.bLoop = bLoop;
+					BlendSpacePlayer.LoopMode = bLoop ? EAnimAssetLoopMode::ForceLoop : EAnimAssetLoopMode::ForceNonLoop;
 					FactoryParams.PushPublicTrait(BlendSpacePlayer);
 
 					FSynchronizeUsingGroupsData SynchronizeUsingGroups;
@@ -135,7 +135,7 @@ namespace UE::UAF
 					AnimSequencePlayer.AnimSequence = Cast<UAnimSequence>(AnimationAsset);
 					AnimSequencePlayer.PlayRate = WantedPlayRate;
 					AnimSequencePlayer.StartPosition =  SelectedTime;
-					AnimSequencePlayer.bLoop = bLoop;
+					AnimSequencePlayer.LoopMode = bLoop ? EAnimAssetLoopMode::ForceLoop : EAnimAssetLoopMode::ForceNonLoop;
 					FactoryParams.PushPublicTrait(AnimSequencePlayer);
 				}
 
