UAF : Expose event name that the montage tick will get added as prerequisite to and fix adding of tick prerequisite after init order changes

diff --git a/Engine/Plugins/Experimental/UAF/UAFAnimGraph/Source/UAFAnimGraph/Internal/Injection/UAFMontageComponent.h b/Engine/Plugins/Experimental/UAF/UAFAnimGraph/Source/UAFAnimGraph/Internal/Injection/UAFMontageComponent.h
--- a/Engine/Plugins/Experimental/UAF/UAFAnimGraph/Source/UAFAnimGraph/Internal/Injection/UAFMontageComponent.h
+++ b/Engine/Plugins/Experimental/UAF/UAFAnimGraph/Source/UAFAnimGraph/Internal/Injection/UAFMontageComponent.h
@@ -94,9 +94,18 @@ private:
 	// The tick function that is used to update montages on the game thread
 	FUAFUpdateMontageTickFunction MontageTickFunction;
 
-	// The event name for the tick function that the montage tick function will be enqueued as prerequisite for
-	// @Todo: Expose this to be able to be set by the user
-	static inline const FLazyName TickDependecyEventName = FLazyName("PrePhysics");
+	// The event name for the module event tick function that the montage tick function will be enqueued as prerequisite for
+	// If None is set, it will fall back to DefaultTickDependencyEventName
+	UPROPERTY(EditAnywhere, Category = Montages, meta = (CustomWidget = AnimNextModuleEvent))
+	FName TickMontagesBeforeEventName = NAME_None;
+
+	// The event name of the module tick function used at runtime to enqueue the montage tick function as a prerequisite to. 
+	// If the the exposed TickMontagesBeforeEventName does not provide an override, it will fall back to DefaultTickDependencyEventName
+	UPROPERTY(Transient)
+	FName TickDependencyEventName = NAME_None;
+
+	// The event name to fall back to for the module event tick function that the montage tick function will be enqueued as prerequisite for
+	static inline const FLazyName DefaultTickDependencyEventName = FLazyName("PrePhysics");
 
 	friend struct UE::UAF::FMontageTrait;
 	friend struct FUAFUpdateMontageTickFunction;
diff --git a/Engine/Plugins/Experimental/UAF/UAFAnimGraph/Source/UAFAnimGraph/Private/Injection/UAFMontageComponent.cpp b/Engine/Plugins/Experimental/UAF/UAFAnimGraph/Source/UAFAnimGraph/Private/Injection/UAFMontageComponent.cpp
--- a/Engine/Plugins/Experimental/UAF/UAFAnimGraph/Source/UAFAnimGraph/Private/Injection/UAFMontageComponent.cpp
+++ b/Engine/Plugins/Experimental/UAF/UAFAnimGraph/Source/UAFAnimGraph/Private/Injection/UAFMontageComponent.cpp
@@ -12,6 +12,7 @@
 #include "Components/SkeletalMeshComponent.h"
 #include "Engine/World.h"
 #include "Engine/Level.h"
+#include "Logging/StructuredLog.h"
 
 
 #include UE_INLINE_GENERATED_CPP_BY_NAME(UAFMontageComponent)
@@ -98,7 +99,14 @@ FUAFMontageComponent::~FUAFMontageComponent()
 	{
 		if (UAnimNextComponent* UAFComponent = Cast<UAnimNextComponent>(ModuleInstance->GetObject()))
 		{
-			UAFComponent->RemovePrerequisite(UAFComponent, MontageTickFunction, TickDependecyEventName);
+			for (UE::UAF::FModuleEventTickFunction& TickFunction : ModuleInstance->GetTickFunctions())
+			{
+				if (TickFunction.Event.GetEventName() == TickDependencyEventName)
+				{
+					TickFunction.RemovePrerequisite(UAFComponent, MontageTickFunction);
+					break;
+				}
+			}
 		}
 	}
 }
@@ -129,8 +137,28 @@ void FUAFMontageComponent::OnBindToInstance()
 		return;
 	}
 
+	TickDependencyEventName = TickMontagesBeforeEventName.IsNone() ? DefaultTickDependencyEventName : TickMontagesBeforeEventName;
+
+	// Depending on init order the SystemReference in the UAFComponnet might not be valid yet, 
+	// so we have to manually find the correct module event tick function and enqueue our prerequisite.
+	UE::UAF::FModuleEventTickFunction* ModuleEventTickFunction = nullptr;
+	for (UE::UAF::FModuleEventTickFunction& TickFunction : ModuleInstance->GetTickFunctions())
+	{
+		if (TickFunction.Event.GetEventName() == TickDependencyEventName)
+		{
+			ModuleEventTickFunction = &TickFunction;
+			break;
+		}
+	}
+
+	if (ModuleEventTickFunction == nullptr)
+	{
+		UE_LOGFMT(LogAnimation, Warning, "UAFMontageComponent: No valid module event tick function found for {EventName}", TickDependencyEventName);
+		return;
+	}
+
 	FAnimNextModuleInstance::RunTaskOnGameThread(
-		[WeakUAFComponent = TWeakObjectPtr<UAnimNextComponent>(UAFComponent), MontageComponent = this]()
+		[WeakUAFComponent = TWeakObjectPtr<UAnimNextComponent>(UAFComponent), ModuleEventTickFunction, MontageComponent = this]()
 		{
 			UAnimNextComponent* UAFComponent = WeakUAFComponent.Get();
 			if (!UAFComponent)
@@ -178,7 +206,7 @@ void FUAFMontageComponent::OnBindToInstance()
 				ULevel* Level = World->PersistentLevel;
 				MontageComponent->MontageTickFunction.RegisterTickFunction(Level);
 
-				UAFComponent->AddPrerequisite(UAFComponent, MontageComponent->MontageTickFunction, TickDependecyEventName);
+				ModuleEventTickFunction->AddPrerequisite(UAFComponent, MontageComponent->MontageTickFunction);
 			}
 		});
 }
