Remove unnecessary reference pose fetching by UAF AnimGraph

diff --git a/Engine/Plugins/Experimental/UAF/UAFAnimGraph/Source/UAFAnimGraph/Private/Graph/AnimNextAnimationGraph.cpp b/Engine/Plugins/Experimental/UAF/UAFAnimGraph/Source/UAFAnimGraph/Private/Graph/AnimNextAnimationGraph.cpp
--- a/Engine/Plugins/Experimental/UAF/UAFAnimGraph/Source/UAFAnimGraph/Private/Graph/AnimNextAnimationGraph.cpp
+++ b/Engine/Plugins/Experimental/UAF/UAFAnimGraph/Source/UAFAnimGraph/Private/Graph/AnimNextAnimationGraph.cpp
@@ -87,11 +87,10 @@ TSharedPtr<FAnimNextGraphInstance> UAnimNextAnimationGraph::AllocateInstance(con
 
 			if (USkeletalMeshComponent* Component = ComponentReference.GetComponent())
 			{
-				UE::UAF::FDataHandle RefPoseHandle = UE::UAF::FDataRegistry::Get()->GetOrGenerateReferencePose(Component);
-				const UE::UAF::FReferencePose& RefPose = RefPoseHandle.GetRef<UE::UAF::FReferencePose>();
-				Context.SetBindingObject(RefPose.SkeletalMeshComponent);
+				Context.SetBindingObject(Component);
 			}
 		}
+
 		InstanceImpl->GraphInstancePtr = Context.AllocateNodeInstance(*InstanceImpl, ResolvedRootTraitHandle);
 	}
 
