Change default blend in UAF blend smoother from linear to hermite cubic

diff --git a/Engine/Plugins/Experimental/UAF/UAFAnimGraph/Source/UAFAnimGraph/Public/Traits/BlendSmoother.h b/Engine/Plugins/Experimental/UAF/UAFAnimGraph/Source/UAFAnimGraph/Public/Traits/BlendSmoother.h
--- a/Engine/Plugins/Experimental/UAF/UAFAnimGraph/Source/UAFAnimGraph/Public/Traits/BlendSmoother.h
+++ b/Engine/Plugins/Experimental/UAF/UAFAnimGraph/Source/UAFAnimGraph/Public/Traits/BlendSmoother.h
@@ -48,7 +48,7 @@ struct FAnimNextBlendSmootherTraitSharedData : public FAnimNextBlendSmootherCore
 
 	/** What type of blend equation to use when converting the time elapsed into a blend weight. */
 	UPROPERTY(EditAnywhere, Category = "Default", meta = (Inline))
-	EAlphaBlendOption BlendType = EAlphaBlendOption::Linear;
+	EAlphaBlendOption BlendType = UE::Anim::DefaultBlendOption;
 
 	/** Custom curve to use when the Custom blend type is used. */
 	UPROPERTY(EditAnywhere, Category = "Default", meta = (Inline))
