Fix crash when drag-dropping an anim sequence into an anim graph
* With the recent changes to the LoopMode, we no longer need to do this work as the pin defaults to a loop mode of Auto

#Synced-CL 48085336

diff --git a/Engine/Plugins/Experimental/UAF/UAFAnimGraph/Source/UAFAnimGraphUncookedOnly/Private/Templates/UAFGraphNodeTemplate_SequencePlayer.h b/Engine/Plugins/Experimental/UAF/UAFAnimGraph/Source/UAFAnimGraphUncookedOnly/Private/Templates/UAFGraphNodeTemplate_SequencePlayer.h
--- a/Engine/Plugins/Experimental/UAF/UAFAnimGraph/Source/UAFAnimGraphUncookedOnly/Private/Templates/UAFGraphNodeTemplate_SequencePlayer.h
+++ b/Engine/Plugins/Experimental/UAF/UAFAnimGraph/Source/UAFAnimGraphUncookedOnly/Private/Templates/UAFGraphNodeTemplate_SequencePlayer.h
@@ -1,4 +1,4 @@
-﻿// Copyright Epic Games, Inc. All Rights Reserved.
+// Copyright Epic Games, Inc. All Rights Reserved.
 
 #pragma once
 
@@ -52,10 +52,6 @@ class UUAFGraphNodeTemplate_SequencePlayer : public UUAFGraphNodeTemplate
 
 			Controller->SetNodeTitle(Node, FText::Format(LOCTEXT("NodeTitleFormat", "Play {0}"), FText::FromString(Asset->GetName())).ToString(), true, true, true);
 
-			// Set looping to asset configuration by default
-			UAnimSequence* AnimSequence = CastChecked<UAnimSequence>(Asset);
-			Controller->SetPinDefaultValue(Node->FindPin(TEXT("AnimNextSequencePlayerTraitSharedData.bLoop")), LexToString(AnimSequence->bLoop), true, true, true, true);
-
 			Controller->CloseUndoBracket();
 		}
 	}
