Crash when exposing variables from Run Graph node
#fix only show AnimNext assets with public (declared) variables when attempting to "Expose Variables" on a RunGraph node
#misc skip setting declared flag for recursive FUtils::GetAssetVariableExports calls (this means that assets will still show up in menu until asset is loaded/resaved)
#Synced-CL 48351162

diff --git a/Engine/Plugins/Experimental/UAF/UAF/Source/UAFUncookedOnly/Private/UncookedOnlyUtils.cpp b/Engine/Plugins/Experimental/UAF/UAF/Source/UAFUncookedOnly/Private/UncookedOnlyUtils.cpp
--- a/Engine/Plugins/Experimental/UAF/UAF/Source/UAFUncookedOnly/Private/UncookedOnlyUtils.cpp
+++ b/Engine/Plugins/Experimental/UAF/UAF/Source/UAFUncookedOnly/Private/UncookedOnlyUtils.cpp
@@ -1105,7 +1105,9 @@ void FUtils::GetAssetVariableExports(const UAnimNextRigVMAssetEditorData* InEdit
 	{
 		if(const UAnimNextVariableEntry* VariableEntry = Cast<UAnimNextVariableEntry>(Entry))
 		{
-			EAnimNextExportedVariableFlags Flags = EAnimNextExportedVariableFlags::Declared | InFlags;
+			const bool bReferencedVariableEntry = EnumHasAnyFlags(InFlags, EAnimNextExportedVariableFlags::Referenced);
+			// Only mark variable as declared when retrieving exports for the asset which actually contains it (non-recursive)
+			EAnimNextExportedVariableFlags Flags = bReferencedVariableEntry ? InFlags : EAnimNextExportedVariableFlags::Declared | InFlags;
 			if(VariableEntry->GetExportAccessSpecifier() == EAnimNextExportAccessSpecifier::Public)
 			{
 				Flags |= EAnimNextExportedVariableFlags::Public;
diff --git a/Engine/Plugins/Experimental/UAF/UAFAnimGraph/Source/UAFAnimGraphEditor/Private/Graph/AnimNextAnimationGraphMenuExtensions.cpp b/Engine/Plugins/Experimental/UAF/UAFAnimGraph/Source/UAFAnimGraphEditor/Private/Graph/AnimNextAnimationGraphMenuExtensions.cpp
--- a/Engine/Plugins/Experimental/UAF/UAFAnimGraph/Source/UAFAnimGraphEditor/Private/Graph/AnimNextAnimationGraphMenuExtensions.cpp
+++ b/Engine/Plugins/Experimental/UAF/UAFAnimGraph/Source/UAFAnimGraphEditor/Private/Graph/AnimNextAnimationGraphMenuExtensions.cpp
@@ -370,23 +370,26 @@ void FAnimationGraphMenuExtensions::RegisterMenus()
 						UAnimNextSharedVariables* Asset = CastChecked<UAnimNextSharedVariables>(InAssetData.GetAsset());
 						NewTrait.InternalAsset = Asset;
 						const FInstancedPropertyBag& VariableDefaults = Asset->GetVariableDefaults();
-						TConstArrayView<FPropertyBagPropertyDesc> Descs = VariableDefaults.GetPropertyBagStruct()->GetPropertyDescs();
-						NewTrait.InternalVariableNames.Reserve(Descs.Num());
-						for(const FPropertyBagPropertyDesc& Desc : Descs)
+						if (VariableDefaults.GetPropertyBagStruct())
 						{
-							if (Desc.CachedProperty->HasAnyPropertyFlags(CPF_NativeAccessSpecifierPublic))
+							TConstArrayView<FPropertyBagPropertyDesc> Descs = VariableDefaults.GetPropertyBagStruct()->GetPropertyDescs();
+							NewTrait.InternalVariableNames.Reserve(Descs.Num());
+							for(const FPropertyBagPropertyDesc& Desc : Descs)
 							{
-								NewTrait.InternalVariableNames.Add(Desc.Name);
+								if (Desc.CachedProperty->HasAnyPropertyFlags(CPF_NativeAccessSpecifierPublic))
+								{
+									NewTrait.InternalVariableNames.Add(Desc.Name);
+								}
 							}
-						}
-						FRigVMTrait_AnimNextPublicVariables::StaticStruct()->ExportText(DefaultValue, &NewTrait, &DefaultTrait, nullptr, PPF_SerializedAsImportText, nullptr);
+							FRigVMTrait_AnimNextPublicVariables::StaticStruct()->ExportText(DefaultValue, &NewTrait, &DefaultTrait, nullptr, PPF_SerializedAsImportText, nullptr);
 
-						const FName ValidTraitName = URigVMSchema::GetUniqueName(Private::VariablesTraitBaseName, [ModelNode](const FName& InName)
-						{
-							return ModelNode->FindPin(InName.ToString()) == nullptr;
-						}, false, false);
+							const FName ValidTraitName = URigVMSchema::GetUniqueName(Private::VariablesTraitBaseName, [ModelNode](const FName& InName)
+							{
+								return ModelNode->FindPin(InName.ToString()) == nullptr;
+							}, false, false);
 
-						VMController->AddTrait(ModelNode->GetFName(), *FRigVMTrait_AnimNextPublicVariables::StaticStruct()->GetPathName(), ValidTraitName, DefaultValue, INDEX_NONE, true, true);
+							VMController->AddTrait(ModelNode->GetFName(), *FRigVMTrait_AnimNextPublicVariables::StaticStruct()->GetPathName(), ValidTraitName, DefaultValue, INDEX_NONE, true, true);
+						}
 					});
 
 					AssetPickerConfig.OnShouldFilterAsset = FOnShouldFilterAsset::CreateLambda([](const FAssetData& InAssetData)
@@ -399,10 +402,10 @@ void FAnimationGraphMenuExtensions::RegisterMenus()
 							{
 								if (const FAnimNextVariableDeclarationData* VariableDeclaration = Export.Data.GetPtr<FAnimNextVariableDeclarationData>())
 								{
-									if (EnumHasAnyFlags(static_cast<EAnimNextExportedVariableFlags>(VariableDeclaration->Flags),EAnimNextExportedVariableFlags::Public))
+									if (EnumHasAllFlags(static_cast<EAnimNextExportedVariableFlags>(VariableDeclaration->Flags),EAnimNextExportedVariableFlags::Public | EAnimNextExportedVariableFlags::Declared))
 									{
 										return false;
-									}					
+									}
 								}
 							}
 						}
@@ -447,4 +450,4 @@ void FAnimationGraphMenuExtensions::UnregisterMenus()
 
 }
 
-#undef LOCTEXT_NAMESPACE
\ No newline at end of file
+#undef LOCTEXT_NAMESPACE
