UAF: Call OpenWorkspaceForObject in CallFunctionSharedDataDetails with UAF Workspace Factory instead of the default

#Synced-CL 48500609

diff --git a/Engine/Plugins/Experimental/UAF/UAFAnimGraph/Source/UAFAnimGraphEditor/Private/Traits/AnimNextCallFunctionSharedDataDetails.cpp b/Engine/Plugins/Experimental/UAF/UAFAnimGraph/Source/UAFAnimGraphEditor/Private/Traits/AnimNextCallFunctionSharedDataDetails.cpp
--- a/Engine/Plugins/Experimental/UAF/UAFAnimGraph/Source/UAFAnimGraphEditor/Private/Traits/AnimNextCallFunctionSharedDataDetails.cpp
+++ b/Engine/Plugins/Experimental/UAF/UAFAnimGraph/Source/UAFAnimGraphEditor/Private/Traits/AnimNextCallFunctionSharedDataDetails.cpp
@@ -1,4 +1,4 @@
-﻿// Copyright Epic Games, Inc. All Rights Reserved.
+// Copyright Epic Games, Inc. All Rights Reserved.
 
 #include "AnimNextCallFunctionSharedDataDetails.h"
 
@@ -19,6 +19,7 @@
 #include "EdGraph/RigVMEdGraphNode.h"
 #include "Modules/ModuleManager.h"
 #include "RigVMModel/RigVMController.h"
+#include "Workspace/AnimNextWorkspaceFactory.h"
 
 #define LOCTEXT_NAMESPACE "FCallFunctionSharedDataDetails"
 
@@ -204,7 +205,7 @@ void FCallFunctionSharedDataDetails::CustomizeChildren(TSharedRef<IPropertyHandl
 
 			// Open the new function's graph
 			UE::Workspace::IWorkspaceEditorModule& WorkspaceEditorModule = FModuleManager::LoadModuleChecked<UE::Workspace::IWorkspaceEditorModule>("WorkspaceEditor");
-			if(UE::Workspace::IWorkspaceEditor* WorkspaceEditor = WorkspaceEditorModule.OpenWorkspaceForObject(CurrentAsset, UE::Workspace::EOpenWorkspaceMethod::Default))
+			if(UE::Workspace::IWorkspaceEditor* WorkspaceEditor = WorkspaceEditorModule.OpenWorkspaceForObject(CurrentAsset, UE::Workspace::EOpenWorkspaceMethod::Default, UAnimNextWorkspaceFactory::StaticClass()))
 			{
 				UObject* EditorObject = EditorData->GetEditorObjectForRigVMGraph(NewFunction->GetContainedGraph());
 				WorkspaceEditor->OpenObjects({EditorObject});
