UAF: Revert sequence player changes from CL 48838940

To address issues with ST queries of timelines prior to relevancy

#Synced-CL 48500753

diff --git a/Engine/Plugins/Experimental/UAF/UAFAnimGraph/Source/UAFAnimGraph/Private/Traits/SequencePlayer.cpp b/Engine/Plugins/Experimental/UAF/UAFAnimGraph/Source/UAFAnimGraph/Private/Traits/SequencePlayer.cpp
--- a/Engine/Plugins/Experimental/UAF/UAFAnimGraph/Source/UAFAnimGraph/Private/Traits/SequencePlayer.cpp
+++ b/Engine/Plugins/Experimental/UAF/UAFAnimGraph/Source/UAFAnimGraph/Private/Traits/SequencePlayer.cpp
@@ -39,20 +39,6 @@ namespace UE::UAF
 		}
 	}
 
-	const UAnimSequence* FSequencePlayerTrait::EvaluateAnimSequence(const FTraitBinding& InBinding) const
-	{
-		FInstanceData* InstanceData = InBinding.GetInstanceData<FInstanceData>();
-		if (!InstanceData->bCachedAnimSequence)
-		{
-			// Only cache this once for each constructed (relevant) sequence player.
-			// Usually this is when OnBecomeRelevant is called, but can be earlier for e.g. timeline queries prior to initial relevancy
-			const FSharedData* SharedData = InBinding.GetSharedData<FSharedData>();
-			InstanceData->AnimSequence = SharedData->GetAnimSequence(InBinding);
-			InstanceData->bCachedAnimSequence = true;
-		}
-		return InstanceData->AnimSequence;
-	}
-
 	void FSequencePlayerTrait::PreEvaluate(FEvaluateTraversalContext& Context, const TTraitBinding<IEvaluate>& Binding) const
 	{
 		const FSharedData* SharedData = Binding.GetSharedData<FSharedData>();
@@ -75,7 +61,7 @@ namespace UE::UAF
 	{
 		const FInstanceData* InstanceData = Binding.GetInstanceData<FInstanceData>();
 
-		if (const UAnimSequence* AnimSeq = EvaluateAnimSequence(Binding))
+		if (InstanceData->AnimSequence)
 		{
 			TTraitBinding<ITimeline> TimelineTrait;
 			Binding.GetStackInterface(TimelineTrait);
@@ -84,7 +70,7 @@ namespace UE::UAF
 			const float PlayRate = State.GetPlayRate();
 			const bool bIsLooping = State.IsLooping();
 
-			auto ExtractRootMotionAttribute = [AnimSequence = AnimSeq, PlayRate, bIsLooping](float StartTime, float DeltaTime, bool bAllowLooping)
+			auto ExtractRootMotionAttribute = [AnimSequence = InstanceData->AnimSequence, PlayRate, bIsLooping](float StartTime, float DeltaTime, bool bAllowLooping)
 			{
 				// We do not check for lifetimes, assume the sequence is alive during pose list execution.
 				check(AnimSequence->IsValidLowLevel());
@@ -107,7 +93,7 @@ namespace UE::UAF
 	void FSequencePlayerTrait::GetSyncMarkers(const FExecutionContext& Context, const TTraitBinding<ITimeline>& Binding, FTimelineSyncMarkerArray& OutSyncMarkers) const
 	{
 		const FInstanceData* InstanceData = Binding.GetInstanceData<FInstanceData>();
-		if (const UAnimSequence* AnimSeq = EvaluateAnimSequence(Binding))
+		if (const UAnimSequence* AnimSeq = InstanceData->AnimSequence.Get())
 		{
 			OutSyncMarkers.Reserve(AnimSeq->AuthoredSyncMarkers.Num());
 
@@ -121,7 +107,7 @@ namespace UE::UAF
 	FTimelineState FSequencePlayerTrait::GetState(const FExecutionContext& Context, const TTraitBinding<ITimeline>& Binding) const
 	{
 		const FInstanceData* InstanceData = Binding.GetInstanceData<FInstanceData>();
-		if (const UAnimSequence* AnimSeq = EvaluateAnimSequence(Binding))
+		if (const UAnimSequence* AnimSeq = InstanceData->AnimSequence.Get())
 		{
 			const FSharedData* SharedData = Binding.GetSharedData<FSharedData>();
 
@@ -138,7 +124,7 @@ namespace UE::UAF
 	FTimelineDelta FSequencePlayerTrait::GetDelta(const FExecutionContext& Context, const TTraitBinding<ITimeline>& Binding) const
 	{
 		const FInstanceData* InstanceData = Binding.GetInstanceData<FInstanceData>();
-		if (const UAnimSequence* AnimSeq = EvaluateAnimSequence(Binding))
+		if (const UAnimSequence* AnimSeq = InstanceData->AnimSequence.Get())
 		{
 			return FTimelineDelta(
 				InstanceData->DeltaTimeRecord.Delta,
@@ -153,7 +139,7 @@ namespace UE::UAF
 		FInstanceData* InstanceData = Binding.GetInstanceData<FInstanceData>();
 
 		// We only advance if we have a valid anim sequence
-		if (const UAnimSequence* AnimSeq = EvaluateAnimSequence(Binding))
+		if (const UAnimSequence* AnimSeq = InstanceData->AnimSequence.Get())
 		{
 			TTraitBinding<ITimeline> TimelineTrait;
 			Binding.GetStackInterface(TimelineTrait);
@@ -179,7 +165,7 @@ namespace UE::UAF
 		FInstanceData* InstanceData = Binding.GetInstanceData<FInstanceData>();
 
 		// Cache the anim sequence we'll play during construction, we don't allow it to change afterwards
-		InstanceData->AnimSequence = EvaluateAnimSequence(Binding);
+		InstanceData->AnimSequence = SharedData->GetAnimSequence(Binding);
 
 		float InternalTimeAccumulator = 0.0f;
 		if (const UAnimSequence* AnimSeq = InstanceData->AnimSequence.Get())
@@ -213,7 +199,7 @@ namespace UE::UAF
 	void FSequencePlayerTrait::GetNotifies(FExecutionContext& Context, const TTraitBinding<INotifySource>& Binding, float StartPosition, float Duration, bool bLooping, TArray<FAnimNotifyEventReference>& OutNotifies) const
 	{
 		const FInstanceData* InstanceData = Binding.GetInstanceData<FInstanceData>();
-		if (const UAnimSequence* AnimSeq = EvaluateAnimSequence(Binding))
+		if (const UAnimSequence* AnimSeq = InstanceData->AnimSequence.Get())
 		{
 			FAnimTickRecord TickRecord;
 			TickRecord.TimeAccumulator = &StartPosition;
diff --git a/Engine/Plugins/Experimental/UAF/UAFAnimGraph/Source/UAFAnimGraph/Private/Traits/SequencePlayer.h b/Engine/Plugins/Experimental/UAF/UAFAnimGraph/Source/UAFAnimGraph/Private/Traits/SequencePlayer.h
--- a/Engine/Plugins/Experimental/UAF/UAFAnimGraph/Source/UAFAnimGraph/Private/Traits/SequencePlayer.h
+++ b/Engine/Plugins/Experimental/UAF/UAFAnimGraph/Source/UAFAnimGraph/Private/Traits/SequencePlayer.h
@@ -40,9 +40,6 @@ namespace UE::UAF
 
 			// The last advance type when FAnimationRuntime::AdvanceTime was called
 			ETypeAdvanceAnim LastAdvanceType = ETAA_Default;
-
-			// Whether we have cached our AnimSequence yet
-			bool bCachedAnimSequence = false;
 		};
 
 		// IEvaluate impl
@@ -68,7 +65,5 @@ namespace UE::UAF
 
 		// INotifySource impl
 		virtual void GetNotifies(FExecutionContext& Context, const TTraitBinding<INotifySource>& Binding, float StartPosition, float Duration, bool bLooping, TArray<FAnimNotifyEventReference>& OutNotifies) const override;
-
-		const UAnimSequence* EvaluateAnimSequence(const FTraitBinding& InBinding) const;
 	};
 }
