UAF: Timelines can now fail state & delta queries

This is to allow clients of the APIs to reason about what to do when querying outside of the 'regular' UpdateGraph call flow (e.g. traversing into not-yet-relevant children).
Calling from PreUpdate into the same stack is a path that should not fail, so I've added ensures for those existing cases.

#Synced-CL 48607479

diff --git a/Engine/Plugins/Experimental/EvaluationNotifies/Source/Runtime/Private/AnimNext/EvaluationNotifiesTrait.cpp b/Engine/Plugins/Experimental/EvaluationNotifies/Source/Runtime/Private/AnimNext/EvaluationNotifiesTrait.cpp
--- a/Engine/Plugins/Experimental/EvaluationNotifies/Source/Runtime/Private/AnimNext/EvaluationNotifiesTrait.cpp
+++ b/Engine/Plugins/Experimental/EvaluationNotifies/Source/Runtime/Private/AnimNext/EvaluationNotifiesTrait.cpp
@@ -42,7 +42,9 @@ namespace UE::UAF
 		ensure(Binding.GetStackInterfaceSuper(TimelinePlayerTrait));
 
 		// Get current state from the stack, advance time, then get the delta state
-		const FTimelineState PreAdvanceState = TimelineTrait.GetState(Context);
+		FTimelineState PreAdvanceState;
+		ensure(TimelineTrait.GetState(Context, PreAdvanceState));
+
 		TimelinePlayerTrait.AdvanceBy(Context, DeltaTime, bDispatchEvents);
 
 		if (!bDispatchEvents)
@@ -53,7 +55,8 @@ namespace UE::UAF
 		TTraitBinding<INotifySource> NotifySourceTrait;
 		ensure(Binding.GetStackInterfaceSuper(NotifySourceTrait));
 
-		const FTimelineDelta Delta = TimelineTrait.GetDelta(Context);
+		FTimelineDelta Delta;
+		ensure(TimelineTrait.GetDelta(Context, Delta));
 
 		// Query for notifies
 		TArray<FAnimNotifyEventReference> Notifies;
diff --git a/Engine/Plugins/Experimental/UAF/UAFAnimGraph/Source/UAFAnimGraph/Internal/Injection/MontageTrait.h b/Engine/Plugins/Experimental/UAF/UAFAnimGraph/Source/UAFAnimGraph/Internal/Injection/MontageTrait.h
--- a/Engine/Plugins/Experimental/UAF/UAFAnimGraph/Source/UAFAnimGraph/Internal/Injection/MontageTrait.h
+++ b/Engine/Plugins/Experimental/UAF/UAFAnimGraph/Source/UAFAnimGraph/Internal/Injection/MontageTrait.h
@@ -149,7 +149,7 @@ namespace UE::UAF
 
 		// ITimeline impl
 		virtual void GetSyncMarkers(const FExecutionContext& Context, const TTraitBinding<ITimeline>& Binding, FTimelineSyncMarkerArray& OutSyncMarkers) const override;
-		virtual FTimelineState GetState(const FExecutionContext& Context, const TTraitBinding<ITimeline>& Binding) const override;
+		virtual bool GetState(const FExecutionContext& Context, const TTraitBinding<ITimeline>& Binding, FTimelineState& OutTimelineState) const override;
 
 		// IGroupSynchronization impl
 		virtual FSyncGroupParameters GetGroupParameters(FExecutionContext& Context, const TTraitBinding<IGroupSynchronization>& Binding) const override;
diff --git a/Engine/Plugins/Experimental/UAF/UAFAnimGraph/Source/UAFAnimGraph/Internal/Traits/BlendStackTrait.h b/Engine/Plugins/Experimental/UAF/UAFAnimGraph/Source/UAFAnimGraph/Internal/Traits/BlendStackTrait.h
--- a/Engine/Plugins/Experimental/UAF/UAFAnimGraph/Source/UAFAnimGraph/Internal/Traits/BlendStackTrait.h
+++ b/Engine/Plugins/Experimental/UAF/UAFAnimGraph/Source/UAFAnimGraph/Internal/Traits/BlendStackTrait.h
@@ -193,7 +193,7 @@ namespace UE::UAF
 		virtual FOnExtractRootMotionAttribute GetOnExtractRootMotionAttribute(FExecutionContext& Context, const TTraitBinding<IAttributeProvider>& Binding) const override;
 
 		// ITimeline impl
-		virtual FTimelineState GetState(const FExecutionContext& Context, const TTraitBinding<ITimeline>& Binding) const override;
+		virtual bool GetState(const FExecutionContext& Context, const TTraitBinding<ITimeline>& Binding, FTimelineState& OutTimelineState) const override;
 
 		// IGarbageCollection impl
 		virtual void AddReferencedObjects(const FExecutionContext& Context, const TTraitBinding<IGarbageCollection>& Binding, FReferenceCollector& Collector) const override;
diff --git a/Engine/Plugins/Experimental/UAF/UAFAnimGraph/Source/UAFAnimGraph/Private/Graph/SyncGroup_GraphInstanceComponent.cpp b/Engine/Plugins/Experimental/UAF/UAFAnimGraph/Source/UAFAnimGraph/Private/Graph/SyncGroup_GraphInstanceComponent.cpp
--- a/Engine/Plugins/Experimental/UAF/UAFAnimGraph/Source/UAFAnimGraph/Private/Graph/SyncGroup_GraphInstanceComponent.cpp
+++ b/Engine/Plugins/Experimental/UAF/UAFAnimGraph/Source/UAFAnimGraph/Private/Graph/SyncGroup_GraphInstanceComponent.cpp
@@ -1311,7 +1311,8 @@ namespace UE::UAF
 			ensure(LeaderTraitStack.GetInterface(GroupSyncTrait));
 
 			// Cache our starting timeline state
-			FTimelineState LeaderStartTimelineState = GroupLeaderContext.TimelineTrait.GetState(Context);
+			FTimelineState LeaderStartTimelineState;
+			ensure(GroupLeaderContext.TimelineTrait.GetState(Context, LeaderStartTimelineState));
 
 			const float PlayRate = LeaderStartTimelineState.GetPlayRate() != 0.0f ? LeaderStartTimelineState.GetPlayRate() : 1.0f;
 
@@ -1328,7 +1329,7 @@ namespace UE::UAF
 				GroupSyncTrait.AdvanceBy(Context, DeltaTime, false);
 
 				// Update our cached state
-				LeaderStartTimelineState = GroupLeaderContext.TimelineTrait.GetState(Context);
+				ensure(GroupLeaderContext.TimelineTrait.GetState(Context, LeaderStartTimelineState));
 			}
 
 			// Record where the leader started from, followers that join the group will start there
@@ -1339,7 +1340,8 @@ namespace UE::UAF
 			GroupSyncTrait.AdvanceBy(Context, DeltaTime, true);
 
 			// Get our new state
-			const FTimelineState LeaderEndTimelineState = GroupLeaderContext.TimelineTrait.GetState(Context);
+			FTimelineState LeaderEndTimelineState;
+			ensure(GroupLeaderContext.TimelineTrait.GetState(Context, LeaderEndTimelineState));
 
 			GroupContext.LeaderTimelineState = LeaderEndTimelineState;
 			GroupContext.bIsLeaderPlayingForward = (DeltaTime * PlayRate) >= 0.0f;
@@ -1519,7 +1521,8 @@ namespace UE::UAF
 			ensure(LeaderTraitStack.GetInterface(GroupSyncTrait));
 
 			// Cache our starting timeline state
-			FTimelineState LeaderStartTimelineState = GroupLeaderContext.TimelineTrait.GetState(Context);
+			FTimelineState LeaderStartTimelineState;
+			ensure(GroupLeaderContext.TimelineTrait.GetState(Context, LeaderStartTimelineState));
 
 			const float PlayRate = LeaderStartTimelineState.GetPlayRate() != 0.0f ? LeaderStartTimelineState.GetPlayRate() : 1.0f;
 
@@ -1574,7 +1577,7 @@ namespace UE::UAF
 					GroupSyncTrait.AdvanceBy(Context, DeltaTime, false);
 
 					// Update our cached state
-					LeaderStartTimelineState = GroupLeaderContext.TimelineTrait.GetState(Context);
+					ensure(GroupLeaderContext.TimelineTrait.GetState(Context, LeaderStartTimelineState));
 				}
 				else
 				{
@@ -1606,7 +1609,8 @@ namespace UE::UAF
 			GroupSyncTrait.AdvanceBy(Context, DeltaTime, true);
 
 			// Compute our new timeline state
-			const FTimelineState LeaderEndTimelineState = GroupLeaderContext.TimelineTrait.GetState(Context);
+			FTimelineState LeaderEndTimelineState;
+			ensure(GroupLeaderContext.TimelineTrait.GetState(Context, LeaderEndTimelineState));
 
 			GroupContext.LeaderTimelineState = LeaderEndTimelineState;
 			GroupContext.bIsLeaderPlayingForward = (DeltaTime * PlayRate) >= 0.0f;
@@ -1736,7 +1740,8 @@ namespace UE::UAF
 			ensure(MemberTraitStack.GetInterface(GroupSyncTrait));
 
 			// Cache our starting timeline state
-			FTimelineState FollowerStartTimelineState = GroupMemberContext.TimelineTrait.GetState(Context);
+			FTimelineState FollowerStartTimelineState;
+			ensure(GroupMemberContext.TimelineTrait.GetState(Context, FollowerStartTimelineState));
 
 			const bool bIsPlayingForward = GroupContext.bIsLeaderPlayingForward;
 			const float PlayRate = FollowerStartTimelineState.GetPlayRate() != 0.0f ? FollowerStartTimelineState.GetPlayRate() : 1.0f;
@@ -1753,7 +1758,7 @@ namespace UE::UAF
 				GroupSyncTrait.AdvanceBy(Context, DeltaTime, false);
 
 				// Update our cached state
-				FollowerStartTimelineState = GroupMemberContext.TimelineTrait.GetState(Context);
+				ensure(GroupMemberContext.TimelineTrait.GetState(Context, FollowerStartTimelineState));
 			}
 
 #if UE_DEBUG_SYNC_GROUPS
@@ -1800,7 +1805,8 @@ namespace UE::UAF
 					GroupContext.bIsLeaderPlayingForward ? TEXT("+") : TEXT("-"),
 					DeltaTime * FMath::Abs(PlayRate));
 
-				const FTimelineState FollowerEndTimelineState = GroupMemberContext.TimelineTrait.GetState(Context);
+				FTimelineState FollowerEndTimelineState;
+				ensure(GroupMemberContext.TimelineTrait.GetState(Context, FollowerEndTimelineState));
 
 				UE_VLOG_UELOG(Context.GetHostObject(), LogAnimMarkerSync, Verbose,
 					TEXT("[%s] [%p] Follow [Time] [%0.2f %s -> %0.2f / %0.2f (%3.2f%%)] Playing [%s]%s"),
@@ -1829,7 +1835,8 @@ namespace UE::UAF
 			ensure(MemberTraitStack.GetInterface(GroupSyncTrait));
 
 			// Cache our starting timeline state
-			FTimelineState FollowerStartTimelineState = GroupMemberContext.TimelineTrait.GetState(Context);
+			FTimelineState FollowerStartTimelineState;
+			ensure(GroupMemberContext.TimelineTrait.GetState(Context, FollowerStartTimelineState));
 
 			const bool bIsPlayingForward = GroupContext.bIsLeaderPlayingForward;
 			const float PlayRate = FollowerStartTimelineState.GetPlayRate() != 0.0f ? FollowerStartTimelineState.GetPlayRate() : 1.0f;
@@ -1931,7 +1938,7 @@ namespace UE::UAF
 				GroupSyncTrait.AdvanceBy(Context, DeltaTime, false);
 
 				// Update our cached state
-				FollowerStartTimelineState = GroupMemberContext.TimelineTrait.GetState(Context);
+				ensure(GroupMemberContext.TimelineTrait.GetState(Context, FollowerStartTimelineState));
 			}
 
 			const float CurrentTime = FollowerStartTimelineState.GetPosition();
@@ -1972,7 +1979,8 @@ namespace UE::UAF
 					GroupContext.bIsLeaderPlayingForward ? TEXT("+") : TEXT("-"),
 					DeltaTime * FMath::Abs(PlayRate));
 
-				const FTimelineState FollowerEndTimelineState = GroupMemberContext.TimelineTrait.GetState(Context);
+				FTimelineState FollowerEndTimelineState;
+				ensure(GroupMemberContext.TimelineTrait.GetState(Context, FollowerEndTimelineState));
 
 				const bool bLooped = GroupContext.bIsLeaderPlayingForward ?
 					(FollowerEndTimelineState.GetPosition() < FollowerStartTimelineState.GetPosition()) :
@@ -2040,7 +2048,8 @@ namespace UE::UAF
 			ensure(MemberTraitStack.GetInterface(GroupSyncTrait));
 
 #if UE_DEBUG_SYNC_GROUPS
-			const FTimelineState FollowerStartTimelineState = GroupMemberContext.TimelineTrait.GetState(Context);
+			FTimelineState FollowerStartTimelineState;
+			ensure(GroupMemberContext.TimelineTrait.GetState(Context, FollowerStartTimelineState));
 #endif
 
 			const float DeltaTime = GroupMember.TraitState.GetDeltaTime();
@@ -2053,7 +2062,8 @@ namespace UE::UAF
 					GroupContext.bIsLeaderPlayingForward ? TEXT("+") : TEXT("-"),
 					DeltaTime * FMath::Abs(FollowerStartTimelineState.GetPlayRate()));
 
-				const FTimelineState FollowerEndTimelineState = GroupMemberContext.TimelineTrait.GetState(Context);
+				FTimelineState FollowerEndTimelineState;
+				ensure(GroupMemberContext.TimelineTrait.GetState(Context, FollowerEndTimelineState));
 
 				UE_VLOG_UELOG(Context.GetHostObject(), LogAnimMarkerSync, Verbose,
 					TEXT("[%s] [%p] Follow [Solo] [%0.2f %s -> %0.2f / %0.2f (%3.2f%%)] Playing [%s]"),
diff --git a/Engine/Plugins/Experimental/UAF/UAFAnimGraph/Source/UAFAnimGraph/Private/Injection/InjectionSiteTrait.cpp b/Engine/Plugins/Experimental/UAF/UAFAnimGraph/Source/UAFAnimGraph/Private/Injection/InjectionSiteTrait.cpp
--- a/Engine/Plugins/Experimental/UAF/UAFAnimGraph/Source/UAFAnimGraph/Private/Injection/InjectionSiteTrait.cpp
+++ b/Engine/Plugins/Experimental/UAF/UAFAnimGraph/Source/UAFAnimGraph/Private/Injection/InjectionSiteTrait.cpp
@@ -264,7 +264,8 @@ namespace UE::UAF
 			TTraitBinding<ITimeline> TimelineTrait;
 			if (Binding.GetStackInterface(TimelineTrait))
 			{
-				const FTimelineState ChildState = TimelineTrait.GetState(Context);
+				FTimelineState ChildState;
+				ensure(TimelineTrait.GetState(Context, ChildState));
 
 				// Compute how much time is left before the timeline ends (can be negative if we overshoot)
 				// Note when looping/infinite we dont have a 'time left', so we leave it at FLT_MAX
diff --git a/Engine/Plugins/Experimental/UAF/UAFAnimGraph/Source/UAFAnimGraph/Private/Injection/MontageTrait.cpp b/Engine/Plugins/Experimental/UAF/UAFAnimGraph/Source/UAFAnimGraph/Private/Injection/MontageTrait.cpp
--- a/Engine/Plugins/Experimental/UAF/UAFAnimGraph/Source/UAFAnimGraph/Private/Injection/MontageTrait.cpp
+++ b/Engine/Plugins/Experimental/UAF/UAFAnimGraph/Source/UAFAnimGraph/Private/Injection/MontageTrait.cpp
@@ -521,10 +521,11 @@ namespace UE::UAF
 		}
 	}
 
-	FTimelineState FMontageTrait::GetState(const FExecutionContext& Context, const TTraitBinding<ITimeline>& Binding) const
+	bool FMontageTrait::GetState(const FExecutionContext& Context, const TTraitBinding<ITimeline>& Binding, FTimelineState& OutTimelineState) const
 	{
 		const FInstanceData* InstanceData = Binding.GetInstanceData<FInstanceData>();
-		return InstanceData->bReturnPreUpdateState ? InstanceData->PreMontageUpdateState : InstanceData->PostMontageUpdateState;
+		OutTimelineState = InstanceData->bReturnPreUpdateState ? InstanceData->PreMontageUpdateState : InstanceData->PostMontageUpdateState;
+		return true;
 	}
 
 	void FMontageTrait::AdvanceBy(FExecutionContext& Context, const TTraitBinding<IGroupSynchronization>& Binding, float DeltaTime, bool bDispatchEvents) const
diff --git a/Engine/Plugins/Experimental/UAF/UAFAnimGraph/Source/UAFAnimGraph/Private/PlayAnim/PlayAnimSlotTrait.cpp b/Engine/Plugins/Experimental/UAF/UAFAnimGraph/Source/UAFAnimGraph/Private/PlayAnim/PlayAnimSlotTrait.cpp
--- a/Engine/Plugins/Experimental/UAF/UAFAnimGraph/Source/UAFAnimGraph/Private/PlayAnim/PlayAnimSlotTrait.cpp
+++ b/Engine/Plugins/Experimental/UAF/UAFAnimGraph/Source/UAFAnimGraph/Private/PlayAnim/PlayAnimSlotTrait.cpp
@@ -229,7 +229,8 @@ namespace UE::UAF
 			TTraitBinding<ITimeline> ChildTimelineTrait;
 			ensure(ChildStack.GetInterface(ChildTimelineTrait));
 
-			const FTimelineState ChildState = ChildTimelineTrait.GetState(Context);
+			FTimelineState ChildState;
+			ensure(ChildTimelineTrait.GetState(Context, ChildState));
 
 			if (InstanceData->CurrentlyActiveRequestIndex == RequestIndex)
 			{
diff --git a/Engine/Plugins/Experimental/UAF/UAFAnimGraph/Source/UAFAnimGraph/Private/TraitInterfaces/ITimeline.cpp b/Engine/Plugins/Experimental/UAF/UAFAnimGraph/Source/UAFAnimGraph/Private/TraitInterfaces/ITimeline.cpp
--- a/Engine/Plugins/Experimental/UAF/UAFAnimGraph/Source/UAFAnimGraph/Private/TraitInterfaces/ITimeline.cpp
+++ b/Engine/Plugins/Experimental/UAF/UAFAnimGraph/Source/UAFAnimGraph/Private/TraitInterfaces/ITimeline.cpp
@@ -30,25 +30,25 @@ namespace UE::UAF
 		}
 	}
 
-	FTimelineState ITimeline::GetState(const FExecutionContext& Context, const TTraitBinding<ITimeline>& Binding) const
+	bool ITimeline::GetState(const FExecutionContext& Context, const TTraitBinding<ITimeline>& Binding, FTimelineState& OutTimelineState) const
 	{
 		TTraitBinding<ITimeline> SuperBinding;
 		if (Binding.GetStackInterfaceSuper(SuperBinding))
 		{
-			return SuperBinding.GetState(Context);
+			return SuperBinding.GetState(Context, OutTimelineState);
 		}
 
-		return FTimelineState();
+		return false;
 	}
 
-	FTimelineDelta ITimeline::GetDelta(const FExecutionContext& Context, const TTraitBinding<ITimeline>& Binding) const
+	bool ITimeline::GetDelta(const FExecutionContext& Context, const TTraitBinding<ITimeline>& Binding, FTimelineDelta& OutTimelineDelta) const
 	{
 		TTraitBinding<ITimeline> SuperBinding;
 		if (Binding.GetStackInterfaceSuper(SuperBinding))
 		{
-			return SuperBinding.GetDelta(Context);
+			return SuperBinding.GetDelta(Context, OutTimelineDelta);
 		}
 
-		return FTimelineDelta();
+		return false;
 	}
 }
diff --git a/Engine/Plugins/Experimental/UAF/UAFAnimGraph/Source/UAFAnimGraph/Private/Traits/BlendSpacePlayer.cpp b/Engine/Plugins/Experimental/UAF/UAFAnimGraph/Source/UAFAnimGraph/Private/Traits/BlendSpacePlayer.cpp
--- a/Engine/Plugins/Experimental/UAF/UAFAnimGraph/Source/UAFAnimGraph/Private/Traits/BlendSpacePlayer.cpp
+++ b/Engine/Plugins/Experimental/UAF/UAFAnimGraph/Source/UAFAnimGraph/Private/Traits/BlendSpacePlayer.cpp
@@ -276,7 +276,7 @@ namespace UE::UAF
 		}
 	}
 
-	FTimelineState FBlendSpacePlayerTrait::GetState(const FExecutionContext& Context, const TTraitBinding<ITimeline>& Binding) const
+	bool FBlendSpacePlayerTrait::GetState(const FExecutionContext& Context, const TTraitBinding<ITimeline>& Binding, FTimelineState& OutTimelineState) const
 	{
 		const FInstanceData* InstanceData = Binding.GetInstanceData<FInstanceData>();
 
@@ -301,11 +301,11 @@ namespace UE::UAF
 				UE::UAF::TTraitBinding<UE::UAF::ITimeline> Timeline;
 				if (ChildTraitStack.GetInterface<UE::UAF::ITimeline>(Timeline))
 				{
-					return Timeline.GetState(Context);
+					return Timeline.GetState(Context, OutTimelineState);
 				}
 			}
 		}
-		return FTimelineState();
+		return false;
 	}
 
 	void FBlendSpacePlayerTrait::AddReferencedObjects(const FExecutionContext& Context, const TTraitBinding<IGarbageCollection>& Binding, FReferenceCollector& Collector) const
diff --git a/Engine/Plugins/Experimental/UAF/UAFAnimGraph/Source/UAFAnimGraph/Private/Traits/BlendSpacePlayer.h b/Engine/Plugins/Experimental/UAF/UAFAnimGraph/Source/UAFAnimGraph/Private/Traits/BlendSpacePlayer.h
--- a/Engine/Plugins/Experimental/UAF/UAFAnimGraph/Source/UAFAnimGraph/Private/Traits/BlendSpacePlayer.h
+++ b/Engine/Plugins/Experimental/UAF/UAFAnimGraph/Source/UAFAnimGraph/Private/Traits/BlendSpacePlayer.h
@@ -86,7 +86,7 @@ namespace UE::UAF
 		virtual void QueueChildrenForTraversal(FUpdateTraversalContext& Context, const TTraitBinding<IUpdateTraversal>& Binding, const FTraitUpdateState& TraitState, FUpdateTraversalQueue& TraversalQueue) const override;
 
 		// ITimeline impl
-		virtual FTimelineState GetState(const FExecutionContext& Context, const TTraitBinding<ITimeline>& Binding) const override;
+		virtual bool GetState(const FExecutionContext& Context, const TTraitBinding<ITimeline>& Binding, FTimelineState& OutTimelineState) const override;
 
 		// IGarbageCollection impl
 		virtual void AddReferencedObjects(const FExecutionContext& Context, const TTraitBinding<IGarbageCollection>& Binding, FReferenceCollector& Collector) const override;
diff --git a/Engine/Plugins/Experimental/UAF/UAFAnimGraph/Source/UAFAnimGraph/Private/Traits/BlendStackTrait.cpp b/Engine/Plugins/Experimental/UAF/UAFAnimGraph/Source/UAFAnimGraph/Private/Traits/BlendStackTrait.cpp
--- a/Engine/Plugins/Experimental/UAF/UAFAnimGraph/Source/UAFAnimGraph/Private/Traits/BlendStackTrait.cpp
+++ b/Engine/Plugins/Experimental/UAF/UAFAnimGraph/Source/UAFAnimGraph/Private/Traits/BlendStackTrait.cpp
@@ -352,7 +352,7 @@ namespace UE::UAF
 		return FOnExtractRootMotionAttribute();
 	}
 
-	FTimelineState FBlendStackCoreTrait::GetState(const FExecutionContext& Context, const TTraitBinding<ITimeline>& Binding) const
+	bool FBlendStackCoreTrait::GetState(const FExecutionContext& Context, const TTraitBinding<ITimeline>& Binding, FTimelineState& OutTimelineState) const
 	{
 		const FInstanceData* InstanceData = Binding.GetInstanceData<FInstanceData>();
 		const int32 CurrentlyActiveGraphIndex = InstanceData->CurrentlyActiveGraphIndex;
@@ -367,15 +367,16 @@ namespace UE::UAF
 			FTraitStackBinding StackBinding;
 			if (IHierarchy::GetForwardedStackInterface<ITimeline>(Context, ChildTraitStack, StackBinding, Timeline))
 			{
-				return Timeline.GetState(Context);
+				return Timeline.GetState(Context, OutTimelineState);
 			}
 			else
 			{
-				return FTimelineState(ActiveGraph.Lifetime, std::numeric_limits<float>::infinity(), 1.0f, false);
+				OutTimelineState = FTimelineState(ActiveGraph.Lifetime, std::numeric_limits<float>::infinity(), 1.0f, false);
+				return true;
 			}
 		}
 
-		return FTimelineState();
+		return false;
 	}
 
 	void FBlendStackCoreTrait::AddReferencedObjects(const FExecutionContext& Context, const TTraitBinding<IGarbageCollection>& Binding, FReferenceCollector& Collector) const
diff --git a/Engine/Plugins/Experimental/UAF/UAFAnimGraph/Source/UAFAnimGraph/Private/Traits/NotifyDispatcher.cpp b/Engine/Plugins/Experimental/UAF/UAFAnimGraph/Source/UAFAnimGraph/Private/Traits/NotifyDispatcher.cpp
--- a/Engine/Plugins/Experimental/UAF/UAFAnimGraph/Source/UAFAnimGraph/Private/Traits/NotifyDispatcher.cpp
+++ b/Engine/Plugins/Experimental/UAF/UAFAnimGraph/Source/UAFAnimGraph/Private/Traits/NotifyDispatcher.cpp
@@ -42,7 +42,8 @@ namespace UE::UAF
 		ensure(Binding.GetStackInterfaceSuper(TimelinePlayerTrait));
 
 		// Get current state from the stack, advance time, then get the delta state
-		const FTimelineState PreAdvanceState = TimelineTrait.GetState(Context);
+		FTimelineState PreAdvanceState;
+		ensure(TimelineTrait.GetState(Context, PreAdvanceState));
 		TimelinePlayerTrait.AdvanceBy(Context, DeltaTime, bDispatchEvents);
 
 		if (!bDispatchEvents)
@@ -58,25 +59,27 @@ namespace UE::UAF
 		TTraitBinding<INotifySource> NotifySourceTrait;
 		ensure(Binding.GetStackInterfaceSuper(NotifySourceTrait));
 
-		const FTimelineDelta Delta = TimelineTrait.GetDelta(Context);
-
-		// Query for notifies
-		TArray<FAnimNotifyEventReference> Notifies;
-		NotifySourceTrait.GetNotifies(Context, PreAdvanceState.GetPosition(), Delta.GetDeltaTime(), PreAdvanceState.IsLooping(), Notifies);
-
-		// Dispatch if we found any
-		if(Notifies.Num() > 0)
+		FTimelineDelta Delta;
+		if (ensure(TimelineTrait.GetDelta(Context, Delta)))
 		{
-			// Ensure we have a handler component on the module
-			if(FAnimNextModuleInstance* ModuleInstance = Context.GetRootGraphInstance().GetModuleInstance())
+			// Query for notifies
+			TArray<FAnimNotifyEventReference> Notifies;
+			NotifySourceTrait.GetNotifies(Context, PreAdvanceState.GetPosition(), Delta.GetDeltaTime(), PreAdvanceState.IsLooping(), Notifies);
+
+			// Dispatch if we found any
+			if(Notifies.Num() > 0)
 			{
-				(void)ModuleInstance->GetComponent<FAnimNextNotifyDispatcherComponent>();
-			}
+				// Ensure we have a handler component on the module
+				if(FAnimNextModuleInstance* ModuleInstance = Context.GetRootGraphInstance().GetModuleInstance())
+				{
+					(void)ModuleInstance->GetComponent<FAnimNextNotifyDispatcherComponent>();
+				}
 
-			auto NotifyDispatchEvent = MakeTraitEvent<FNotifyDispatchEvent>();
-			NotifyDispatchEvent->Notifies = MoveTemp(Notifies);
-			NotifyDispatchEvent->Weight = 1.0f;
-			Context.RaiseOutputTraitEvent(NotifyDispatchEvent);
+				auto NotifyDispatchEvent = MakeTraitEvent<FNotifyDispatchEvent>();
+				NotifyDispatchEvent->Notifies = MoveTemp(Notifies);
+				NotifyDispatchEvent->Weight = 1.0f;
+				Context.RaiseOutputTraitEvent(NotifyDispatchEvent);
+			}
 		}
 	}
 
diff --git a/Engine/Plugins/Experimental/UAF/UAFAnimGraph/Source/UAFAnimGraph/Private/Traits/SequencePlayer.cpp b/Engine/Plugins/Experimental/UAF/UAFAnimGraph/Source/UAFAnimGraph/Private/Traits/SequencePlayer.cpp
--- a/Engine/Plugins/Experimental/UAF/UAFAnimGraph/Source/UAFAnimGraph/Private/Traits/SequencePlayer.cpp
+++ b/Engine/Plugins/Experimental/UAF/UAFAnimGraph/Source/UAFAnimGraph/Private/Traits/SequencePlayer.cpp
@@ -66,7 +66,9 @@ namespace UE::UAF
 			TTraitBinding<ITimeline> TimelineTrait;
 			Binding.GetStackInterface(TimelineTrait);
 
-			const FTimelineState State = TimelineTrait.GetState(Context);
+			FTimelineState State;
+			ensure(TimelineTrait.GetState(Context, State));
+
 			const float PlayRate = State.GetPlayRate();
 			const bool bIsLooping = State.IsLooping();
 
@@ -83,7 +85,7 @@ namespace UE::UAF
 
 				return AnimSequence->ExtractRootMotion(FAnimExtractContext(static_cast<double>(StartTime), true, FDeltaTimeRecord(DeltaTime), bAllowLooping && bIsLooping));
 			};
-
+			
 			return FOnExtractRootMotionAttribute::CreateLambda(ExtractRootMotionAttribute);
 		}
 
@@ -104,34 +106,36 @@ namespace UE::UAF
 		}
 	}
 
-	FTimelineState FSequencePlayerTrait::GetState(const FExecutionContext& Context, const TTraitBinding<ITimeline>& Binding) const
+	bool FSequencePlayerTrait::GetState(const FExecutionContext& Context, const TTraitBinding<ITimeline>& Binding, FTimelineState& OutTimelineState) const
 	{
 		const FInstanceData* InstanceData = Binding.GetInstanceData<FInstanceData>();
 		if (const UAnimSequence* AnimSeq = InstanceData->AnimSequence.Get())
 		{
 			const FSharedData* SharedData = Binding.GetSharedData<FSharedData>();
 
-			return FTimelineState(
+			OutTimelineState = FTimelineState(
 				InstanceData->InternalTimeAccumulator,
 				AnimSeq->GetPlayLength(),
 				SharedData->GetPlayRate(Binding) * AnimSeq->RateScale,
 				IsSequenceLoopingHelper(AnimSeq, SharedData->GetLoopMode(Binding))).WithDebugName(AnimSeq->GetFName());
+			return true;
 		}
 
-		return FTimelineState();
+		return false;
 	}
 
-	FTimelineDelta FSequencePlayerTrait::GetDelta(const FExecutionContext& Context, const TTraitBinding<ITimeline>& Binding) const
+	bool FSequencePlayerTrait::GetDelta(const FExecutionContext& Context, const TTraitBinding<ITimeline>& Binding, FTimelineDelta& OutTimelineDelta) const
 	{
 		const FInstanceData* InstanceData = Binding.GetInstanceData<FInstanceData>();
 		if (const UAnimSequence* AnimSeq = InstanceData->AnimSequence.Get())
 		{
-			return FTimelineDelta(
+			OutTimelineDelta = FTimelineDelta(
 				InstanceData->DeltaTimeRecord.Delta,
 				InstanceData->LastAdvanceType);
+			return true;
 		}
 
-		return FTimelineDelta();
+		return false;
 	}
 
 	void FSequencePlayerTrait::AdvanceBy(FExecutionContext& Context, const TTraitBinding<ITimelinePlayer>& Binding, float DeltaTime, bool bDispatchEvents) const
@@ -144,7 +148,8 @@ namespace UE::UAF
 			TTraitBinding<ITimeline> TimelineTrait;
 			Binding.GetStackInterface(TimelineTrait);
 
-			const FTimelineState State = TimelineTrait.GetState(Context);
+			FTimelineState State;
+			ensure(TimelineTrait.GetState(Context, State));
 
 			const float PlayRate = State.GetPlayRate();
 			const bool bIsLooping = State.IsLooping();
diff --git a/Engine/Plugins/Experimental/UAF/UAFAnimGraph/Source/UAFAnimGraph/Private/Traits/SequencePlayer.h b/Engine/Plugins/Experimental/UAF/UAFAnimGraph/Source/UAFAnimGraph/Private/Traits/SequencePlayer.h
--- a/Engine/Plugins/Experimental/UAF/UAFAnimGraph/Source/UAFAnimGraph/Private/Traits/SequencePlayer.h
+++ b/Engine/Plugins/Experimental/UAF/UAFAnimGraph/Source/UAFAnimGraph/Private/Traits/SequencePlayer.h
@@ -50,8 +50,8 @@ namespace UE::UAF
 
 		// ITimeline impl
 		virtual void GetSyncMarkers(const FExecutionContext& Context, const TTraitBinding<ITimeline>& Binding, FTimelineSyncMarkerArray& OutSyncMarkers) const override;
-		virtual FTimelineState GetState(const FExecutionContext& Context, const TTraitBinding<ITimeline>& Binding) const override;
-		virtual FTimelineDelta GetDelta(const FExecutionContext& Context, const TTraitBinding<ITimeline>& Binding) const override;
+		virtual bool GetState(const FExecutionContext& Context, const TTraitBinding<ITimeline>& Binding, FTimelineState& OutTimelineState) const override;
+		virtual bool GetDelta(const FExecutionContext& Context, const TTraitBinding<ITimeline>& Binding, FTimelineDelta& OutTimelineDelta) const override;
 
 		// ITimelinePlayer impl
 		virtual void AdvanceBy(FExecutionContext& Context, const TTraitBinding<ITimelinePlayer>& Binding, float DeltaTime, bool bDispatchEvents) const override;
diff --git a/Engine/Plugins/Experimental/UAF/UAFAnimGraph/Source/UAFAnimGraph/Public/TraitInterfaces/ITimeline.h b/Engine/Plugins/Experimental/UAF/UAFAnimGraph/Source/UAFAnimGraph/Public/TraitInterfaces/ITimeline.h
--- a/Engine/Plugins/Experimental/UAF/UAFAnimGraph/Source/UAFAnimGraph/Public/TraitInterfaces/ITimeline.h
+++ b/Engine/Plugins/Experimental/UAF/UAFAnimGraph/Source/UAFAnimGraph/Public/TraitInterfaces/ITimeline.h
@@ -230,10 +230,14 @@ namespace UE::UAF
 		UE_API virtual void GetSyncMarkers(const FExecutionContext& Context, const TTraitBinding<ITimeline>& Binding, FTimelineSyncMarkerArray& OutSyncMarkers) const;
 
 		// Returns the state of this timeline
-		UE_API virtual FTimelineState GetState(const FExecutionContext& Context, const TTraitBinding<ITimeline>& Binding) const;
+		// It is possible for the timeline to be in an 'unknown' state if this is called before supporting data is copied/read, in which case this 
+		// function will return false and the state will not be filled
+		UE_API virtual bool GetState(const FExecutionContext& Context, const TTraitBinding<ITimeline>& Binding, FTimelineState& OutTimelineState) const;
 
 		// Returns the last change in state of this timeline
-		UE_API virtual FTimelineDelta GetDelta(const FExecutionContext& Context, const TTraitBinding<ITimeline>& Binding) const;
+		// It is possible for the timeline to be in an 'unknown' state if this is called before supporting data is copied/read, in which case this 
+		// function will return false and the delta will not be filled
+		UE_API virtual bool GetDelta(const FExecutionContext& Context, const TTraitBinding<ITimeline>& Binding, FTimelineDelta& OutTimelineDelta) const;
 
 #if WITH_EDITOR
 		UE_API virtual const FText& GetDisplayName() const override;
@@ -254,15 +258,15 @@ namespace UE::UAF
 		}
 
 		// @see ITimeline::GetState
-		FTimelineState GetState(const FExecutionContext& Context) const
+		bool GetState(const FExecutionContext& Context, FTimelineState& OutTimelineState) const
 		{
-			return GetInterface()->GetState(Context, *this);
+			return GetInterface()->GetState(Context, *this, OutTimelineState);
 		}
 
 		// @see ITimeline::GetDelta
-		FTimelineDelta GetDelta(const FExecutionContext& Context) const
+		bool GetDelta(const FExecutionContext& Context, FTimelineDelta& OutTimelineDelta) const
 		{
-			return GetInterface()->GetDelta(Context, *this);
+			return GetInterface()->GetDelta(Context, *this, OutTimelineDelta);
 		}
 
 	protected:
diff --git a/Engine/Plugins/Experimental/UAF/UAFPoseSearch/Source/UAFPoseSearch/Private/MotionMatchingTrait.cpp b/Engine/Plugins/Experimental/UAF/UAFPoseSearch/Source/UAFPoseSearch/Private/MotionMatchingTrait.cpp
--- a/Engine/Plugins/Experimental/UAF/UAFPoseSearch/Source/UAFPoseSearch/Private/MotionMatchingTrait.cpp
+++ b/Engine/Plugins/Experimental/UAF/UAFPoseSearch/Source/UAFPoseSearch/Private/MotionMatchingTrait.cpp
@@ -158,35 +158,38 @@ PRAGMA_ENABLE_DEPRECATION_WARNINGS
 			// We adjust the MotionMatchingInteractionState time to the current player node's asset time. This is done 
 			// because the player node may have ticked more or less time than we expected due to variable dt or the 
 			// dynamic playback rate adjustment and as such the MotionMatchingInteractionState does not update by itself
-			const FTimelineState TimelineState = TimelineTrait.GetState(Context);
-			const float TimelineRealTime = TimelineState.GetPosition();
-			float TimelineNormalizedTime = TimelineRealTime;
-
-			// @todo: This is a hack, since we don't have a way to get the normalized timeline of a blendspace yet (pending syncing...?)
-			if (Cast<UBlendSpace>(MotionMatchingState.SearchResult.GetAnimationAssetForRole()))
+			float TimelineNormalizedTime = 0.0f;
+			FTimelineState TimelineState;
+			if (TimelineTrait.GetState(Context, TimelineState))
 			{
-				// Convert to normalized time.
-				const float TimelineDuration = TimelineState.GetDuration();
-				if (TimelineDuration > 0.0f)
+				const float TimelineRealTime = TimelineState.GetPosition();
+				TimelineNormalizedTime = TimelineRealTime;
+
+				// @todo: This is a hack, since we don't have a way to get the normalized timeline of a blendspace yet (pending syncing...?)
+				if (Cast<UBlendSpace>(MotionMatchingState.SearchResult.GetAnimationAssetForRole()))
 				{
-					// NOTE: This doesn't work because AnimNext normalized time does not match database indexing normalized time.
-					TimelineNormalizedTime = TimelineRealTime / TimelineDuration;
+					// Convert to normalized time.
+					const float TimelineDuration = TimelineState.GetDuration();
+					if (TimelineDuration > 0.0f)
+					{
+						// NOTE: This doesn't work because AnimNext normalized time does not match database indexing normalized time.
+						TimelineNormalizedTime = TimelineRealTime / TimelineDuration;
+					}
 				}
-			}
 
 #if ENABLE_VISUAL_LOG
-			if (FVisualLogger::IsRecording())
-			{
-				const FName AnimName = TimelineState.GetDebugName();
-				const float Duration = TimelineState.GetDuration();
-				static const TCHAR* LogName = TEXT("FMotionMatchingTrait");
-				UE_VLOG(Context.GetHostObject(), LogName, Verbose, TEXT("TimelineRealTime: %f"), TimelineRealTime);
-				UE_VLOG(Context.GetHostObject(), LogName, Verbose, TEXT("TimelineNormalizedTime: %f"), TimelineNormalizedTime);
-				UE_VLOG(Context.GetHostObject(), LogName, Verbose, TEXT("AnimName: %s"), *AnimName.ToString());
-				UE_VLOG(Context.GetHostObject(), LogName, Verbose, TEXT("Duration: %f"), Duration);
-			}
+				if (FVisualLogger::IsRecording())
+				{
+					const FName AnimName = TimelineState.GetDebugName();
+					const float Duration = TimelineState.GetDuration();
+					static const TCHAR* LogName = TEXT("FMotionMatchingTrait");
+					UE_VLOG(Context.GetHostObject(), LogName, Verbose, TEXT("TimelineRealTime: %f"), TimelineRealTime);
+					UE_VLOG(Context.GetHostObject(), LogName, Verbose, TEXT("TimelineNormalizedTime: %f"), TimelineNormalizedTime);
+					UE_VLOG(Context.GetHostObject(), LogName, Verbose, TEXT("AnimName: %s"), *AnimName.ToString());
+					UE_VLOG(Context.GetHostObject(), LogName, Verbose, TEXT("Duration: %f"), Duration);
+				}
 #endif
-
+			}
 			// @todo: ask blendstack what it's playing instead of relying on the previous MotionMatchingState.SearchResult
 			// We adjust the motion matching state asset time to the current player node's asset time. This is done 
 			// because the player node may have ticked more or less time than we expected due to variable dt or the 
@@ -220,7 +223,7 @@ PRAGMA_ENABLE_DEPRECATION_WARNINGS
 
 				FChooserEvaluationContext ChooserEvaluationContext(const_cast<UObject*>(AnimContext));
 				const FUAFAssetInstance& Instance = Binding.GetTraitPtr().GetNodeInstance()->GetOwner();
-                ChooserEvaluationContext.AddStructViewParam(FStructView::Make(const_cast<FUAFAssetInstance&>(Instance)));
+				ChooserEvaluationContext.AddStructViewParam(FStructView::Make(const_cast<FUAFAssetInstance&>(Instance)));
 				SearchContext.AddRole(GetCommonDefaultRole(DatabasesToSearch), &ChooserEvaluationContext, PoseHistory);
 
 				// If we can't advance or enough time has elapsed since the last pose jump then search
@@ -251,7 +254,7 @@ PRAGMA_ENABLE_DEPRECATION_WARNINGS
 					SearchResult = SearchContext.GetContinuingPoseSearchResult();
 					SearchResult.bIsContinuingPoseSearch = true;
 
-			#if UE_POSE_SEARCH_TRACE_ENABLED
+#if UE_POSE_SEARCH_TRACE_ENABLED
 					// in case we skipped the search, we still have to track we would have requested to evaluate DatabasesToSearch and SearchResult.Database
 					for (const TObjectPtr<const UPoseSearchDatabase>& DatabaseToSearch : DatabasesToSearch)
 					{
@@ -259,7 +262,7 @@ PRAGMA_ENABLE_DEPRECATION_WARNINGS
 					}
 
 					SearchContext.Track(SearchResult.Database.Get());
-			#endif // UE_POSE_SEARCH_TRACE_ENABLED
+#endif // UE_POSE_SEARCH_TRACE_ENABLED
 				}
 
 				const float WantedPlayRate = CalculateWantedPlayRate(SearchResult, SearchContext, PoseSearchPlayRate, PoseHistory ? PoseHistory->GetTrajectorySpeedMultiplier() : 1.f, EventToSearch);
diff --git a/Engine/Plugins/Experimental/UAF/UAFPoseSearch/Source/UAFPoseSearch/Private/PoseSearchResultEmulatorTrait.cpp b/Engine/Plugins/Experimental/UAF/UAFPoseSearch/Source/UAFPoseSearch/Private/PoseSearchResultEmulatorTrait.cpp
--- a/Engine/Plugins/Experimental/UAF/UAFPoseSearch/Source/UAFPoseSearch/Private/PoseSearchResultEmulatorTrait.cpp
+++ b/Engine/Plugins/Experimental/UAF/UAFPoseSearch/Source/UAFPoseSearch/Private/PoseSearchResultEmulatorTrait.cpp
@@ -83,22 +83,25 @@ namespace UE::UAF
 			if (bHasValidActiveGraph)
 			{
 				CurrentAnim = ActiveGraphRequest->FactoryObject;
-				const FTimelineState TimelineState = TimelineTrait.GetState(Context);
-				float TimelinePosition = TimelineState.GetPosition();
-
-				const bool bCurrentAnimIsBlendSpace = CurrentAnim->IsA<UBlendSpace>();
-				if (bCurrentAnimIsBlendSpace)
+				FTimelineState TimelineState;
+				if(TimelineTrait.GetState(Context, TimelineState))
 				{
-					// Convert to normalized time.
-					const float TimelineDuration = TimelineState.GetDuration();
-					if (TimelineDuration > 0.0f)
+					float TimelinePosition = TimelineState.GetPosition();
+
+					const bool bCurrentAnimIsBlendSpace = CurrentAnim->IsA<UBlendSpace>();
+					if (bCurrentAnimIsBlendSpace)
 					{
-						// NOTE: This doesn't work because AnimNext normalized time does not match database indexing normalized time.
-						TimelinePosition /= TimelineState.GetDuration();
+						// Convert to normalized time.
+						const float TimelineDuration = TimelineState.GetDuration();
+						if (TimelineDuration > 0.0f)
+						{
+							// NOTE: This doesn't work because AnimNext normalized time does not match database indexing normalized time.
+							TimelinePosition /= TimelineState.GetDuration();
+						}
 					}
-				}
 
-				TimeDelta = FMath::Abs(TimelinePosition - SelectedTime);
+					TimeDelta = FMath::Abs(TimelinePosition - SelectedTime);
+				}
 			}
 
 			const float MaxTimeDeltaAllowed = SharedData->GetMaxTimeDeltaAllowed(Binding);
diff --git a/Engine/Plugins/Experimental/UAF/UAFStateTree/Source/UAFStateTree/Private/AnimNextStateTreeContext.cpp b/Engine/Plugins/Experimental/UAF/UAFStateTree/Source/UAFStateTree/Private/AnimNextStateTreeContext.cpp
--- a/Engine/Plugins/Experimental/UAF/UAFStateTree/Source/UAFStateTree/Private/AnimNextStateTreeContext.cpp
+++ b/Engine/Plugins/Experimental/UAF/UAFStateTree/Source/UAFStateTree/Private/AnimNextStateTreeContext.cpp
@@ -42,7 +42,11 @@ float FAnimNextStateTreeTraitContext::QueryPlaybackRatio() const
 	UE::UAF::TTraitBinding<UE::UAF::ITimeline> Timeline;
 	if (Binding->GetStackInterface<UE::UAF::ITimeline>(Timeline))
 	{
-		return Timeline.GetState(*Context).GetPositionRatio();
+		UE::UAF::FTimelineState TimelineState;
+		if (Timeline.GetState(*Context, TimelineState))
+		{
+			return TimelineState.GetPositionRatio();
+		}
 	}
 
 	return 1.0f;
@@ -53,7 +57,11 @@ float FAnimNextStateTreeTraitContext::QueryTimeLeft() const
 	UE::UAF::TTraitBinding<UE::UAF::ITimeline> Timeline;
 	if (Binding->GetStackInterface<UE::UAF::ITimeline>(Timeline))
 	{
-		return Timeline.GetState(*Context).GetTimeLeft();
+		UE::UAF::FTimelineState TimelineState;
+		if (Timeline.GetState(*Context, TimelineState))
+		{
+			return TimelineState.GetTimeLeft();
+		}
 	}
 
 	return std::numeric_limits<float>::infinity();
@@ -64,7 +72,11 @@ float FAnimNextStateTreeTraitContext::QueryDuration() const
 	UE::UAF::TTraitBinding<UE::UAF::ITimeline> Timeline;
 	if (Binding->GetStackInterface<UE::UAF::ITimeline>(Timeline))
 	{
-		return Timeline.GetState(*Context).GetDuration();
+		UE::UAF::FTimelineState TimelineState;
+		if (Timeline.GetState(*Context, TimelineState))
+		{
+			return TimelineState.GetDuration();
+		}
 	}
 
 	return 0.0f;
@@ -75,7 +87,11 @@ bool FAnimNextStateTreeTraitContext::QueryIsLooping() const
 	UE::UAF::TTraitBinding<UE::UAF::ITimeline> Timeline;
 	if (Binding->GetStackInterface<UE::UAF::ITimeline>(Timeline))
 	{
-		return Timeline.GetState(*Context).IsLooping();
+		UE::UAF::FTimelineState TimelineState;
+		if (Timeline.GetState(*Context, TimelineState))
+		{
+			return TimelineState.IsLooping();
+		}
 	}
 
 	return false;
diff --git a/Engine/Plugins/Experimental/UAF/UAFWarping/Source/UAFWarping/Private/SteeringTrait.cpp b/Engine/Plugins/Experimental/UAF/UAFWarping/Source/UAFWarping/Private/SteeringTrait.cpp
--- a/Engine/Plugins/Experimental/UAF/UAFWarping/Source/UAFWarping/Private/SteeringTrait.cpp
+++ b/Engine/Plugins/Experimental/UAF/UAFWarping/Source/UAFWarping/Private/SteeringTrait.cpp
@@ -68,9 +68,10 @@ namespace UE::UAF
 		// Get Current Anim Asset Time
 		{
 			TTraitBinding<ITimeline> TimelineTrait;
-			if (Binding.GetStackInterface(TimelineTrait))
+			FTimelineState TimelineState;
+			if (Binding.GetStackInterface(TimelineTrait) && TimelineTrait.GetState(Context, TimelineState))
 			{
-				InstanceData->CurrentAnimAssetTime = TimelineTrait.GetState(Context).GetPosition();
+				InstanceData->CurrentAnimAssetTime = TimelineState.GetPosition();
 			}
 			else
 			{
