Fix IBlendStack::GetActiveGraphInstance

#Synced-CL 48625956

diff --git a/Engine/Plugins/Experimental/UAF/UAFAnimGraph/Source/UAFAnimGraph/Internal/Traits/BlendStackTrait.h b/Engine/Plugins/Experimental/UAF/UAFAnimGraph/Source/UAFAnimGraph/Internal/Traits/BlendStackTrait.h
--- a/Engine/Plugins/Experimental/UAF/UAFAnimGraph/Source/UAFAnimGraph/Internal/Traits/BlendStackTrait.h
+++ b/Engine/Plugins/Experimental/UAF/UAFAnimGraph/Source/UAFAnimGraph/Internal/Traits/BlendStackTrait.h
@@ -201,6 +201,7 @@ namespace UE::UAF
 		// IBlendStack impl
 		virtual int32 PushGraph(FExecutionContext& Context, const TTraitBinding<IBlendStack>& Binding, IBlendStack::FGraphRequest&& GraphRequest) const override;
 		virtual int32 GetActiveGraph(FExecutionContext& Context, const TTraitBinding<IBlendStack>& Binding, IBlendStack::FGraphRequestPtr& OutGraphRequest) const override;
+		virtual int32 GetActiveGraphInstance(FExecutionContext& Context, const TTraitBinding<IBlendStack>& Binding, FAnimNextGraphInstance*& OutGraphInstance) const override;
 		virtual IBlendStack::FGraphRequestPtr GetGraph(FExecutionContext& Context, const TTraitBinding<IBlendStack>& Binding, int32 InChildIndex) const override;
 		virtual FAnimNextGraphInstance* GetGraphInstance(FExecutionContext& Context, const TTraitBinding<IBlendStack>& Binding, int32 InChildIndex) const override;
 
diff --git a/Engine/Plugins/Experimental/UAF/UAFAnimGraph/Source/UAFAnimGraph/Private/Traits/BlendStackTrait.cpp b/Engine/Plugins/Experimental/UAF/UAFAnimGraph/Source/UAFAnimGraph/Private/Traits/BlendStackTrait.cpp
--- a/Engine/Plugins/Experimental/UAF/UAFAnimGraph/Source/UAFAnimGraph/Private/Traits/BlendStackTrait.cpp
+++ b/Engine/Plugins/Experimental/UAF/UAFAnimGraph/Source/UAFAnimGraph/Private/Traits/BlendStackTrait.cpp
@@ -339,14 +339,16 @@ namespace UE::UAF
 				}
 			});
 
-			auto ExtractRootMotionAttribute = [AnimSequence](float StartTime, float DeltaTime, bool bAllowLooping)
+			if (AnimSequence)
 			{
-				// We do not check for lifetimes, assume the sequence is alive during pose list execution.
-				check(AnimSequence->IsValidLowLevel());
-				return AnimSequence->ExtractRootMotion(FAnimExtractContext(static_cast<double>(StartTime), true, FDeltaTimeRecord(DeltaTime), bAllowLooping && AnimSequence->bLoop)); 
-			};
-				
-			return FOnExtractRootMotionAttribute::CreateLambda(ExtractRootMotionAttribute);
+				auto ExtractRootMotionAttribute = [AnimSequence](float StartTime, float DeltaTime, bool bAllowLooping)
+					{
+						// We do not check for lifetimes, assume the sequence is alive during pose list execution.
+						check(AnimSequence->IsValidLowLevel());
+						return AnimSequence->ExtractRootMotion(FAnimExtractContext(static_cast<double>(StartTime), true, FDeltaTimeRecord(DeltaTime), bAllowLooping && AnimSequence->bLoop)); 
+					};
+				return FOnExtractRootMotionAttribute::CreateLambda(ExtractRootMotionAttribute);
+			}
 		}
 
 		return FOnExtractRootMotionAttribute();
@@ -466,6 +468,21 @@ namespace UE::UAF
 		return INDEX_NONE;
 	}
 
+	int32 FBlendStackCoreTrait::GetActiveGraphInstance(FExecutionContext& Context, const TTraitBinding<IBlendStack>& Binding,
+		FAnimNextGraphInstance*& OutGraphInstance) const
+	{
+		FInstanceData* InstanceData = Binding.GetInstanceData<FInstanceData>();
+		const int32 CurrentlyActiveGraphIndex = InstanceData->CurrentlyActiveGraphIndex;
+		if (CurrentlyActiveGraphIndex != INDEX_NONE)
+		{
+			FGraphState& GraphState = InstanceData->ChildGraphs[CurrentlyActiveGraphIndex];
+			OutGraphInstance = GraphState.Instance.Get();
+			return CurrentlyActiveGraphIndex;
+		}
+
+		return INDEX_NONE;
+	}
+
 	IBlendStack::FGraphRequestPtr FBlendStackCoreTrait::GetGraph(FExecutionContext& Context, const TTraitBinding<IBlendStack>& Binding, int32 InChildIndex) const
 	{
 		FInstanceData* InstanceData = Binding.GetInstanceData<FInstanceData>();
