Ensure that no instance of an AnimNextGraph are created if it has not succesfully been compiled
#Synced-CL 48648404

diff --git a/Engine/Plugins/Experimental/UAF/UAFAnimGraph/Source/UAFAnimGraph/Private/Graph/AnimNextAnimationGraph.cpp b/Engine/Plugins/Experimental/UAF/UAFAnimGraph/Source/UAFAnimGraph/Private/Graph/AnimNextAnimationGraph.cpp
--- a/Engine/Plugins/Experimental/UAF/UAFAnimGraph/Source/UAFAnimGraph/Private/Graph/AnimNextAnimationGraph.cpp
+++ b/Engine/Plugins/Experimental/UAF/UAFAnimGraph/Source/UAFAnimGraph/Private/Graph/AnimNextAnimationGraph.cpp
@@ -56,6 +56,13 @@ TSharedPtr<FAnimNextGraphInstance> UAnimNextAnimationGraph::AllocateInstance(con
 	{
 		return nullptr;
 	}
+	
+#if WITH_EDITOR
+	if (CompilationState == EAnimNextRigVMAssetState::CompiledWithErrors)
+	{
+		return nullptr;
+	}
+#endif
 
 	TSharedPtr<FAnimNextGraphInstance> InstanceImpl = MakeShared<FAnimNextGraphInstance>();
 	InstanceImpl->Asset = this;
