Update UAF ST to use CallFunction Function Picker
- Add CallFunction function picker dedicated type.
- Update UAF ST Tasks & conditions to use that type
- Note: Not deprecating old properties in CallFunction trait for now as we decide best practice to deprecate in UAF.

#tests editor, PIE

#Synced-CL 48737380

diff --git a/Engine/Plugins/Experimental/UAF/UAF/Source/UAFUncookedOnly/Internal/AnimNextRigVMAssetEditorData.h b/Engine/Plugins/Experimental/UAF/UAF/Source/UAFUncookedOnly/Internal/AnimNextRigVMAssetEditorData.h
--- a/Engine/Plugins/Experimental/UAF/UAF/Source/UAFUncookedOnly/Internal/AnimNextRigVMAssetEditorData.h
+++ b/Engine/Plugins/Experimental/UAF/UAF/Source/UAFUncookedOnly/Internal/AnimNextRigVMAssetEditorData.h
@@ -66,6 +66,7 @@ namespace UE::UAF::Editor
 	class FVariableProxyCustomization;
 	class FAnimNextAnimGraphEditorModule;
 	class FCallFunctionSharedDataDetails;
+	class FCallFunctionInfoDetails;
 	class FAssetCompilationHandler;
 	class FVariableBindingPropertyCustomization;
 	struct FFunctionsOutlinerEntryItem;
@@ -263,6 +264,7 @@ protected:
 	friend struct UE::UAF::UncookedOnly::FPublicVariablesImpl;
 	friend UE::UAF::Editor::FAnimNextAnimGraphEditorModule;
 	friend class UE::UAF::Editor::FCallFunctionSharedDataDetails;
+	friend class UE::UAF::Editor::FCallFunctionInfoDetails;
 	friend UAnimNextWorkspaceEditorMode;
 	friend UE::UAF::UncookedOnly::FCompilationScope;
 	friend UE::UAF::UncookedOnly::Compilation;
diff --git a/Engine/Plugins/Experimental/UAF/UAFAnimGraph/Source/UAFAnimGraph/Internal/Traits/CallFunction.h b/Engine/Plugins/Experimental/UAF/UAFAnimGraph/Source/UAFAnimGraph/Internal/Traits/CallFunction.h
--- a/Engine/Plugins/Experimental/UAF/UAFAnimGraph/Source/UAFAnimGraph/Internal/Traits/CallFunction.h
+++ b/Engine/Plugins/Experimental/UAF/UAFAnimGraph/Source/UAFAnimGraph/Internal/Traits/CallFunction.h
@@ -8,6 +8,8 @@
 #include "TraitInterfaces/IUpdate.h"
 #include "CallFunction.generated.h"
 
+#define UE_API UAFANIMGRAPH_API
+
 UENUM()
 enum class EAnimNextCallFunctionCallSite : uint8
 {
@@ -21,12 +23,36 @@ enum class EAnimNextCallFunctionCallSite : uint8
 	PostUpdate,
 };
 
+
+/** Struct used to help call / pick functions across UAF */
+USTRUCT(meta = (DisplayName = "UAF Function"))
+struct FUAFCallFunctionInfo 
+{
+	GENERATED_BODY()
+
+#if WITH_EDITORONLY_DATA
+	UPROPERTY(VisibleAnywhere, Category = "Call Function", meta = (PinHiddenByDefault, Hidden))
+	FRigVMGraphFunctionHeader FunctionHeader;
+#endif
+
+	/** The function to call */
+	UPROPERTY(EditAnywhere, Category = "Call Function", meta = (PinHiddenByDefault, Hidden))
+	FName Function;
+
+	/** Internal event name derived from Function above */
+	UPROPERTY(VisibleAnywhere, Category = "Call Function", meta = (PinHiddenByDefault, Hidden))
+	FName FunctionEvent;
+};
+
+
 /** A trait that calls a RigVM function at the specified update point */
 USTRUCT(meta = (DisplayName = "Call RigVM Function", AllowDuplicates = "true"))
 struct FAnimNextCallFunctionSharedData : public FAnimNextTraitSharedData
 {
 	GENERATED_BODY()
 
+	/** @TODO: Replace these properties with FUAFCallFunctionInfo once we determine UAF deprecation path */
+
 #if WITH_EDITORONLY_DATA
 	UPROPERTY(VisibleAnywhere, Category = "Call Function", meta=(Hidden))
 	FRigVMGraphFunctionHeader FunctionHeader;
@@ -78,3 +104,5 @@ struct FCallFunctionTrait : FAdditiveTrait, IUpdate
 	void CallFunctionForMatchingSite(const TTraitBinding<IUpdate>& InBinding, const FTraitUpdateState& InTraitState, EAnimNextCallFunctionCallSite InCallSite) const;
 };
 }
+
+#undef UE_API
\ No newline at end of file
diff --git a/Engine/Plugins/Experimental/UAF/UAFAnimGraph/Source/UAFAnimGraph/Private/AnimNextAnimGraphModule.cpp b/Engine/Plugins/Experimental/UAF/UAFAnimGraph/Source/UAFAnimGraph/Private/AnimNextAnimGraphModule.cpp
--- a/Engine/Plugins/Experimental/UAF/UAFAnimGraph/Source/UAFAnimGraph/Private/AnimNextAnimGraphModule.cpp
+++ b/Engine/Plugins/Experimental/UAF/UAFAnimGraph/Source/UAFAnimGraph/Private/AnimNextAnimGraphModule.cpp
@@ -30,6 +30,7 @@
 #include "Native/UAFSimplePrePhysicsGraphComponent.h"
 #include "Traits/BlendSmoother.h"
 #include "Traits/BlendStackTrait.h"
+#include "Traits/CallFunction.h"
 #include "Graph/AnimNextAnimGraph.h"
 
 
@@ -75,6 +76,7 @@ void FAnimNextAnimGraphModule::StartupModule()
 		FAnimNextInjectionBlendSettings::StaticStruct(),
 		FModifyCurveParameters::StaticStruct(),
 		FAnimNextBoneMapping::StaticStruct(),
+		FUAFCallFunctionInfo::StaticStruct(),
 	};
 
 	RigVMRegistry.RegisterStructTypes(AllowedStructTypes);
diff --git a/Engine/Plugins/Experimental/UAF/UAFAnimGraph/Source/UAFAnimGraphEditor/Private/AnimNextAnimGraphEditorModule.cpp b/Engine/Plugins/Experimental/UAF/UAFAnimGraph/Source/UAFAnimGraphEditor/Private/AnimNextAnimGraphEditorModule.cpp
--- a/Engine/Plugins/Experimental/UAF/UAFAnimGraph/Source/UAFAnimGraphEditor/Private/AnimNextAnimGraphEditorModule.cpp
+++ b/Engine/Plugins/Experimental/UAF/UAFAnimGraph/Source/UAFAnimGraphEditor/Private/AnimNextAnimGraphEditorModule.cpp
@@ -25,6 +25,7 @@
 #include "Graph/RigUnit_AnimNextGraphRoot.h"
 #include "Graph/TraitEditorTabSummoner.h"
 #include "Traits/AnimNextCallFunctionSharedDataDetails.h"
+#include "Traits/UAFCallFunctionInfoDetails.h"
 #include "Traits/CallFunction.h"
 #include "WorkflowOrientedApp/WorkflowTabManager.h"
 #include "IWorkspaceEditor.h"
@@ -83,6 +84,9 @@ class FAnimNextAnimGraphEditorModule : public IModuleInterface
 		PropertyModule.RegisterCustomPropertyTypeLayout(
 			FAnimNextCallFunctionSharedData::StaticStruct()->GetFName(),
 			FOnGetPropertyTypeCustomizationInstance::CreateLambda([] { return MakeShared<FCallFunctionSharedDataDetails>(); }));
+		PropertyModule.RegisterCustomPropertyTypeLayout(
+			FUAFCallFunctionInfo::StaticStruct()->GetFName(),
+			FOnGetPropertyTypeCustomizationInstance::CreateLambda([] { return MakeShared<FCallFunctionInfoDetails>(); }));
 		PropertyModule.RegisterCustomPropertyTypeLayout(
 					FAnimNextAnimGraph::StaticStruct()->GetFName(),
 					FOnGetPropertyTypeCustomizationInstance::CreateLambda([] { return MakeShared<FAnimNextGraphDetails>(); }));
diff --git a/Engine/Plugins/Experimental/UAF/UAFAnimGraph/Source/UAFAnimGraphEditor/Private/Traits/UAFCallFunctionInfoDetails.cpp b/Engine/Plugins/Experimental/UAF/UAFAnimGraph/Source/UAFAnimGraphEditor/Private/Traits/UAFCallFunctionInfoDetails.cpp
new file mode 100644
--- /dev/null
+++ b/Engine/Plugins/Experimental/UAF/UAFAnimGraph/Source/UAFAnimGraphEditor/Private/Traits/UAFCallFunctionInfoDetails.cpp
@@ -0,0 +1,214 @@
+// Copyright Epic Games, Inc. All Rights Reserved.
+
+#include "UAFCallFunctionInfoDetails.h"
+
+#include "DetailLayoutBuilder.h"
+#include "DetailWidgetRow.h"
+#include "IDetailChildrenBuilder.h"
+#include "IDetailPropertyRow.h"
+#include "IPropertyUtilities.h"
+#include "PropertyHandle.h"
+#include "UncookedOnlyUtils.h"
+#include "Traits/CallFunction.h"
+#include "Common/SRigVMFunctionPicker.h"
+#include "AnimNextRigVMAsset.h"
+#include "AnimNextRigVMAssetEditorData.h"
+#include "IWorkspaceEditor.h"
+#include "IWorkspaceEditorModule.h"
+#include "ScopedTransaction.h"
+#include "EdGraph/RigVMEdGraphNode.h"
+#include "Modules/ModuleManager.h"
+#include "RigVMModel/RigVMController.h"
+#include "UAFCompilationScope.h"
+#include "Workspace/AnimNextWorkspaceFactory.h"
+
+#define LOCTEXT_NAMESPACE "FCallFunctionInfoDetails"
+
+namespace UE::UAF::Editor
+{
+
+void FCallFunctionInfoDetails::CustomizeHeader(TSharedRef<IPropertyHandle> InPropertyHandle, FDetailWidgetRow& InHeaderRow, IPropertyTypeCustomizationUtils& InCustomizationUtils)
+{
+	using namespace UE::UAF::UncookedOnly;
+
+	PropertyHandle = InPropertyHandle;
+	FunctionPropertyHandle = InPropertyHandle->GetChildHandle(GET_MEMBER_NAME_CHECKED(FUAFCallFunctionInfo, Function));
+	FunctionHeaderPropertyHandle = InPropertyHandle->GetChildHandle(GET_MEMBER_NAME_CHECKED(FUAFCallFunctionInfo, FunctionHeader));
+	FunctionEventPropertyHandle = InPropertyHandle->GetChildHandle(GET_MEMBER_NAME_CHECKED(FUAFCallFunctionInfo, FunctionEvent));
+
+	TArray<TWeakObjectPtr<UObject>> Objects = InCustomizationUtils.GetPropertyUtilities()->GetSelectedObjects();
+	UObject* RigVMAssetSubObject = Objects.Num() ? Objects[0].Get() : nullptr;
+	UAnimNextRigVMAsset* CurrentAsset = nullptr;
+	if (RigVMAssetSubObject)
+	{
+		CurrentAsset = RigVMAssetSubObject->GetTypedOuter<UAnimNextRigVMAsset>();
+	}
+
+	auto OnFunctionPicked = [this, WeakCurrentAsset = TWeakObjectPtr<UAnimNextRigVMAsset>(CurrentAsset)](const FRigVMGraphFunctionHeader& InFunctionHeader)
+	{
+		UAnimNextRigVMAsset* CurrentAsset = WeakCurrentAsset.Get();
+		if(CurrentAsset == nullptr)
+		{
+			return;
+		}
+
+		FCompilationScope CompilerResults(LOCTEXT("ModifiedAssets_FunctionPicked", "Modified Function Picked"), MakeConstArrayView({ CurrentAsset }));
+		FScopedTransaction Transaction(LOCTEXT("SetFunctionTransaction", "Set Function"));
+
+		// Update function event & header from function name
+		FunctionPropertyHandle->SetValue(InFunctionHeader.IsValid() ? InFunctionHeader.Name : NAME_None);
+		FunctionEventPropertyHandle->SetValue(InFunctionHeader.IsValid() ? FName(*UncookedOnly::FUtils::MakeFunctionWrapperEventName(InFunctionHeader.Name)): NAME_None);
+
+		FunctionHeaderPropertyHandle->NotifyPreChange();
+		FunctionHeaderPropertyHandle->EnumerateRawData([&InFunctionHeader](void* InRawData, const int32 InDataIndex, const int32 InNumDatas)
+		{
+			FRigVMGraphFunctionHeader* FunctionHeader = static_cast<FRigVMGraphFunctionHeader*>(InRawData);
+			if(InFunctionHeader.IsValid())
+			{
+				*FunctionHeader = InFunctionHeader;
+			}
+			else
+			{
+				*FunctionHeader = FRigVMGraphFunctionHeader();
+			}
+			return true;
+		});
+		FunctionHeaderPropertyHandle->NotifyPostChange(EPropertyChangeType::ValueSet);
+		FunctionHeaderPropertyHandle->NotifyFinishedChangingProperties();
+	};
+	InHeaderRow
+	.NameContent()
+	[
+		FunctionPropertyHandle->CreatePropertyNameWidget()
+	]
+	.ValueContent()
+	[
+		SNew(SRigVMFunctionPicker)
+		.CurrentAsset(FAssetData(CurrentAsset))
+		.FunctionName_Lambda([this]()
+		{
+			bool bMultipleValues = false;
+			TOptional<FRigVMGraphFunctionHeader> Header;
+			PropertyHandle->EnumerateConstRawData([&Header, &bMultipleValues](const void* InRawData, const int32 InDataIndex, const int32 InNumDatas)
+			{
+				const FUAFCallFunctionInfo* FunctionInfo = static_cast<const FUAFCallFunctionInfo*>(InRawData);
+				if (FunctionInfo)
+				{
+					if (!Header.IsSet())
+					{
+						Header = FunctionInfo->FunctionHeader;
+					}
+					else if (Header.GetValue() != FunctionInfo->FunctionHeader)
+					{
+						Header = FRigVMGraphFunctionHeader();
+						bMultipleValues = true;
+						return false;
+					}
+				}
+				return true;
+			});
+
+			if(bMultipleValues)
+			{
+				return LOCTEXT("MultipleValuesLabel", "Multiple Values");
+			}
+			else if(Header.IsSet() && Header.GetValue().IsValid())
+			{
+				return FText::FromName(Header.GetValue().Name);
+			}
+			return LOCTEXT("NoFunctionSelectedLabel", "None");
+		})
+		.FunctionToolTip_Lambda([this]()
+		{
+			bool bMultipleValues = false;
+			TOptional<FRigVMGraphFunctionHeader> Header;
+			PropertyHandle->EnumerateConstRawData([&Header, &bMultipleValues](const void* InRawData, const int32 InDataIndex, const int32 InNumDatas)
+			{
+				const FUAFCallFunctionInfo* FunctionInfo = static_cast<const FUAFCallFunctionInfo*>(InRawData);
+				if (FunctionInfo)
+				{
+					if (!Header.IsSet())
+					{
+						Header = FunctionInfo->FunctionHeader;
+					}
+					else if (Header.GetValue() != FunctionInfo->FunctionHeader)
+					{
+						Header = FRigVMGraphFunctionHeader();
+						bMultipleValues = true;
+						return false;
+					}
+				}
+				return true;
+			});
+
+			if(bMultipleValues)
+			{
+				return LOCTEXT("MultipleValuesLabel", "Multiple Values");
+			}
+			else if(Header.IsSet() && Header.GetValue().IsValid())
+			{
+				static const FTextFormat ToolTipFormat(LOCTEXT("FunctionToolTipFormat", "{0}\n{1}"));
+				
+				return Header.GetValue().Description.Len() > 0 ?
+					FText::FromString(Header.GetValue().Description) :
+					FText::Format(ToolTipFormat, FText::FromString(Header.GetValue().LibraryPointer.GetFunctionName()), FText::FromString(Header.GetValue().LibraryPointer.GetLibraryNodePath()));
+			}
+			return LOCTEXT("NoFunctionSelectedLabel", "None");
+		})
+		.OnRigVMFunctionPicked_Lambda(OnFunctionPicked)
+		.OnNewFunction_Lambda([this, WeakCurrentAsset = TWeakObjectPtr<UAnimNextRigVMAsset>(CurrentAsset), OnFunctionPicked]()
+		{
+			UAnimNextRigVMAsset* CurrentAsset = WeakCurrentAsset.Get();
+			if (CurrentAsset == nullptr)
+			{
+				return;
+			}
+
+			UAnimNextRigVMAssetEditorData* EditorData = UncookedOnly::FUtils::GetEditorData<UAnimNextRigVMAssetEditorData>(CurrentAsset);
+			URigVMLibraryNode* NewFunction = nullptr;
+			{
+				FCompilationScope CompilerResults(LOCTEXT("ModifiedAssets_FunctionAdded", "Modified Asset Added Function"), MakeConstArrayView({ CurrentAsset }));
+				FScopedTransaction Transaction(LOCTEXT("AddFunctionTransaction", "Add Function"));
+				NewFunction = EditorData->AddFunction(TEXT("NewFunction"), true, true, true);
+
+				if(NewFunction == nullptr)
+				{
+					return;
+				}
+
+				const FRigVMGraphFunctionData* FunctionData = EditorData->GraphFunctionStore.FindFunction(NewFunction->GetFunctionIdentifier());
+				if(FunctionData == nullptr)
+				{
+					return;
+				}
+
+				// Set our function
+				OnFunctionPicked(FunctionData->Header);
+			}
+
+			// Open the new function's graph
+			UE::Workspace::IWorkspaceEditorModule& WorkspaceEditorModule = FModuleManager::LoadModuleChecked<UE::Workspace::IWorkspaceEditorModule>("WorkspaceEditor");
+			if(UE::Workspace::IWorkspaceEditor* WorkspaceEditor = WorkspaceEditorModule.OpenWorkspaceForObject(CurrentAsset, UE::Workspace::EOpenWorkspaceMethod::Default, UAnimNextWorkspaceFactory::StaticClass()))
+			{
+				UObject* EditorObject = EditorData->GetEditorObjectForRigVMGraph(NewFunction->GetContainedGraph());
+				WorkspaceEditor->OpenObjects({EditorObject});
+			}
+		})
+	];
+}
+
+void FCallFunctionInfoDetails::CustomizeChildren(TSharedRef<IPropertyHandle> InPropertyHandle, IDetailChildrenBuilder& InChildBuilder, IPropertyTypeCustomizationUtils& InCustomizationUtils)
+{
+	InChildBuilder.AddProperty(FunctionPropertyHandle.ToSharedRef())
+	.Visibility(EVisibility::Collapsed);
+
+	InChildBuilder.AddProperty(FunctionHeaderPropertyHandle.ToSharedRef())
+	.Visibility(EVisibility::Collapsed);
+
+	InChildBuilder.AddProperty(FunctionEventPropertyHandle.ToSharedRef())
+	.Visibility(EVisibility::Collapsed);
+}
+
+}
+
+#undef LOCTEXT_NAMESPACE
\ No newline at end of file
diff --git a/Engine/Plugins/Experimental/UAF/UAFAnimGraph/Source/UAFAnimGraphEditor/Private/Traits/UAFCallFunctionInfoDetails.h b/Engine/Plugins/Experimental/UAF/UAFAnimGraph/Source/UAFAnimGraphEditor/Private/Traits/UAFCallFunctionInfoDetails.h
new file mode 100644
--- /dev/null
+++ b/Engine/Plugins/Experimental/UAF/UAFAnimGraph/Source/UAFAnimGraphEditor/Private/Traits/UAFCallFunctionInfoDetails.h
@@ -0,0 +1,22 @@
+// Copyright Epic Games, Inc. All Rights Reserved.
+
+#pragma once
+
+#include "IPropertyTypeCustomization.h"
+
+namespace UE::UAF::Editor
+{
+
+class FCallFunctionInfoDetails : public IPropertyTypeCustomization
+{
+	// IPropertyTypeCustomization interface
+	virtual void CustomizeHeader(TSharedRef<IPropertyHandle> InPropertyHandle, FDetailWidgetRow& InHeaderRow, IPropertyTypeCustomizationUtils& InCustomizationUtils) override;
+	virtual void CustomizeChildren(TSharedRef<IPropertyHandle> InPropertyHandle, IDetailChildrenBuilder& InChildBuilder, IPropertyTypeCustomizationUtils& InCustomizationUtils) override;
+
+	TSharedPtr<IPropertyHandle> PropertyHandle;
+	TSharedPtr<IPropertyHandle> FunctionPropertyHandle;
+	TSharedPtr<IPropertyHandle> FunctionHeaderPropertyHandle;
+	TSharedPtr<IPropertyHandle> FunctionEventPropertyHandle;
+};
+
+}
\ No newline at end of file
diff --git a/Engine/Plugins/Experimental/UAF/UAFStateTree/Source/UAFStateTree/Internal/Conditions/AnimNextStateTreeRigVMConditionBase.h b/Engine/Plugins/Experimental/UAF/UAFStateTree/Source/UAFStateTree/Internal/Conditions/AnimNextStateTreeRigVMConditionBase.h
--- a/Engine/Plugins/Experimental/UAF/UAFStateTree/Source/UAFStateTree/Internal/Conditions/AnimNextStateTreeRigVMConditionBase.h
+++ b/Engine/Plugins/Experimental/UAF/UAFStateTree/Source/UAFStateTree/Internal/Conditions/AnimNextStateTreeRigVMConditionBase.h
@@ -7,6 +7,7 @@
 #include "RigVMCore/RigVMGraphFunctionDefinition.h"
 #include "StateTreeConditionBase.h"
 #include "StateTreeReference.h"
+#include "Traits/CallFunction.h"
 
 #include "AnimNextStateTreeRigVMConditionBase.generated.h"
 
@@ -21,6 +22,8 @@ struct FStateTreeExecutionContext;
 struct FAnimNextStateTreeTraitContext;
 struct FRigVMClient;
 
+#define UE_API UAFSTATETREE_API
+
 USTRUCT()
 struct FAnimNextStateTreeRigVMConditionInstanceData
 {
@@ -37,7 +40,7 @@ struct FAnimNextStateTreeRigVMConditionInstanceData
 /**
  * Wrapper for RigVM based Conditions. 
  */
-USTRUCT()
+USTRUCT(MinimalAPI)
 struct FAnimNextStateTreeRigVMConditionBase : public FStateTreeConditionBase
 {
 	GENERATED_BODY()
@@ -67,12 +70,19 @@ public:
 public:
 
 #if WITH_EDITORONLY_DATA
-	UPROPERTY(EditAnywhere, Category = "Parameter", Meta = (GetOptions = "UAFStateTreeUncookedOnly.AnimNextStateTreeFunctionLibraryHelper.GetExposedAnimNextFunctionNames"))
+	UPROPERTY(meta = (DeprecatedProperty, DeprecationMessage = "RigVMFunctionHeader is deprecated, use CallFunctionInfo.FunctionHeader"))
+	FRigVMGraphFunctionHeader RigVMFunctionHeader = FRigVMGraphFunctionHeader();
+
+	UPROPERTY(meta = (DeprecatedProperty, DeprecationMessage = "ConditionFunctionName is deprecated, use CallFunctionInfo.Function"))
 	FName ConditionFunctionName = NAME_None;
+
+	UPROPERTY(meta = (DeprecatedProperty, DeprecationMessage = "InternalEventName is deprecated, use CallFunctionInfo.FunctionEvent"))
+	FName InternalEventName = NAME_None;
 #endif
 
-	UPROPERTY(EditAnywhere, Category = "Parameters")
-	FRigVMGraphFunctionHeader RigVMFunctionHeader = FRigVMGraphFunctionHeader();
+	/** Info needed to call a UAF function. Also has a custom function picker. */
+	UPROPERTY(EditAnywhere, Category = "Call Function", meta = (Hidden))
+	FUAFCallFunctionInfo CallFunctionInfo;
 
 	/** Owning state name, populated during programmatic graph creation */
 	UPROPERTY()
@@ -82,12 +92,24 @@ public:
 	UPROPERTY()
 	FGuid NodeId = FGuid();
 
-	UPROPERTY()
-	FName InternalEventName = NAME_None;
-
 	UPROPERTY()
 	int32 ResultIndex = INDEX_NONE;
 
+	UE_API void PostSerialize(const FArchive& Ar);
+
 public:
 	TStateTreeExternalDataHandle<FAnimNextStateTreeTraitContext> TraitContextHandle;
-};
\ No newline at end of file
+};
+
+#if WITH_EDITORONLY_DATA
+template<>
+struct TStructOpsTypeTraits<FAnimNextStateTreeRigVMConditionBase> : public TStructOpsTypeTraitsBase2<FAnimNextStateTreeRigVMConditionBase>
+{
+	enum
+	{
+		WithPostSerialize = true,
+	};
+};
+#endif // WITH_EDITORONLY_DATA
+
+#undef UE_API
diff --git a/Engine/Plugins/Experimental/UAF/UAFStateTree/Source/UAFStateTree/Internal/Tasks/AnimNextStateTreeRigVMTaskBase.h b/Engine/Plugins/Experimental/UAF/UAFStateTree/Source/UAFStateTree/Internal/Tasks/AnimNextStateTreeRigVMTaskBase.h
--- a/Engine/Plugins/Experimental/UAF/UAFStateTree/Source/UAFStateTree/Internal/Tasks/AnimNextStateTreeRigVMTaskBase.h
+++ b/Engine/Plugins/Experimental/UAF/UAFStateTree/Source/UAFStateTree/Internal/Tasks/AnimNextStateTreeRigVMTaskBase.h
@@ -6,6 +6,7 @@
 #include "AnimNextStateTreeTypes.h"
 #include "RigVMCore/RigVMGraphFunctionDefinition.h"
 #include "StateTreeTaskBase.h"
+#include "Traits/CallFunction.h"
 
 #include "AnimNextStateTreeRigVMTaskBase.generated.h"
 
@@ -20,8 +21,10 @@ struct FStateTreeExecutionContext;
 struct FRigVMClient;
 struct FAnimNextStateTreeTraitContext;
 
+#define UE_API UAFSTATETREE_API
+
 USTRUCT()
-struct UAFSTATETREE_API FAnimNextStateTreeRigVMTaskInstanceData
+struct UE_API FAnimNextStateTreeRigVMTaskInstanceData
 {
 	GENERATED_BODY()
 
@@ -36,8 +39,8 @@ struct UAFSTATETREE_API FAnimNextStateTreeRigVMTaskInstanceData
 /**
  * Wrapper for RigVM based Tasks. 
  */
-USTRUCT()
-struct UAFSTATETREE_API FAnimNextStateTreeRigVMTaskBase : public FAnimNextStateTreeTaskBase
+USTRUCT(MinimalAPI)
+struct FAnimNextStateTreeRigVMTaskBase : public FAnimNextStateTreeTaskBase
 {
 	GENERATED_BODY()
 
@@ -66,12 +69,19 @@ public:
 public:
 
 #if WITH_EDITORONLY_DATA
-	UPROPERTY(EditAnywhere, Category = "Parameter", Meta = (GetOptions = "UAFStateTreeUncookedOnly.AnimNextStateTreeFunctionLibraryHelper.GetExposedAnimNextFunctionNames"))
+	UPROPERTY(meta = (DeprecatedProperty, DeprecationMessage = "RigVMFunctionHeader is deprecated, use CallFunctionInfo.FunctionHeader"))
+	FRigVMGraphFunctionHeader RigVMFunctionHeader = FRigVMGraphFunctionHeader();
+	
+	UPROPERTY(meta = (DeprecatedProperty, DeprecationMessage = "TaskFunctionName is deprecated, use CallFunctionInfo.Function"))
 	FName TaskFunctionName = NAME_None;
+	
+	UPROPERTY(meta = (DeprecatedProperty, DeprecationMessage = "InternalEventName is deprecated, use CallFunctionInfo.FunctionEvent"))
+	FName InternalEventName = NAME_None;
 #endif
 
-	UPROPERTY(EditAnywhere, Category = "Parameters")
-	FRigVMGraphFunctionHeader RigVMFunctionHeader = FRigVMGraphFunctionHeader();
+	/** Info needed to call a UAF function. Also has a custom function picker. */
+	UPROPERTY(EditAnywhere, Category = "Call Function", meta = (Hidden))
+	FUAFCallFunctionInfo CallFunctionInfo;
 
 	/** Owning state name, populated during programmatic graph creation */
 	UPROPERTY()
@@ -81,12 +91,24 @@ public:
 	UPROPERTY()
 	FGuid NodeId = FGuid();
 
-	UPROPERTY()
-	FName InternalEventName = NAME_None;
-
 	UPROPERTY()
 	int32 ResultIndex = INDEX_NONE;
 
+	UE_API void PostSerialize(const FArchive& Ar);
+
 public:
 	TStateTreeExternalDataHandle<FAnimNextStateTreeTraitContext> TraitContextHandle;
-};
\ No newline at end of file
+};
+
+#if WITH_EDITORONLY_DATA
+template<>
+struct TStructOpsTypeTraits<FAnimNextStateTreeRigVMTaskBase> : public TStructOpsTypeTraitsBase2<FAnimNextStateTreeRigVMTaskBase>
+{
+	enum
+	{
+		WithPostSerialize = true,
+	};
+};
+#endif // WITH_EDITORONLY_DATA
+
+#undef UE_API
\ No newline at end of file
diff --git a/Engine/Plugins/Experimental/UAF/UAFStateTree/Source/UAFStateTree/Private/Conditions/AnimNextStateTreeRigVMConditionBase.cpp b/Engine/Plugins/Experimental/UAF/UAFStateTree/Source/UAFStateTree/Private/Conditions/AnimNextStateTreeRigVMConditionBase.cpp
--- a/Engine/Plugins/Experimental/UAF/UAFStateTree/Source/UAFStateTree/Private/Conditions/AnimNextStateTreeRigVMConditionBase.cpp
+++ b/Engine/Plugins/Experimental/UAF/UAFStateTree/Source/UAFStateTree/Private/Conditions/AnimNextStateTreeRigVMConditionBase.cpp
@@ -50,7 +50,7 @@ bool FAnimNextStateTreeRigVMConditionBase::TestCondition(FStateTreeExecutionCont
 			if (!InstanceData->FunctionHandle.IsValid())
 #endif
 			{
-				InstanceData->FunctionHandle = RigVMInstance->GetFunctionHandle(InternalEventName);
+				InstanceData->FunctionHandle = RigVMInstance->GetFunctionHandle(CallFunctionInfo.FunctionEvent);
 			}
 
 			if (!ensure(InstanceData->FunctionHandle.IsValid() && ResultIndex != INDEX_NONE))
@@ -62,7 +62,7 @@ bool FAnimNextStateTreeRigVMConditionBase::TestCondition(FStateTreeExecutionCont
 			FRigVMExtendedExecuteContext& ExtendedExecuteContext = RigVMComponent.GetExtendedExecuteContext();
 			FAnimNextExecuteContext& AnimNextContext = ExtendedExecuteContext.GetPublicDataSafe<FAnimNextExecuteContext>();
 
-			FAnimNextGraphContextData ContextData(GraphInstance, InternalEventName, AnimNextContext.GetDeltaTime<float>());
+			FAnimNextGraphContextData ContextData(GraphInstance, CallFunctionInfo.FunctionEvent, AnimNextContext.GetDeltaTime<float>());
 			UE::UAF::FScopedExecuteContextData ContextDataScope(AnimNextContext, ContextData);
 
 			RigVMInstance->CallFunctionHandle(InstanceData->FunctionHandle, ExtendedExecuteContext, InstanceData->Parameters);
@@ -83,42 +83,17 @@ FText FAnimNextStateTreeRigVMConditionBase::GetDescription(const FGuid& ID, FSta
 
 void FAnimNextStateTreeRigVMConditionBase::PostEditNodeChangeChainProperty(const FPropertyChangedChainEvent& PropertyChangedEvent, FStateTreeDataView InstanceDataView)
 {
-	if (PropertyChangedEvent.GetMemberPropertyName() == GET_MEMBER_NAME_CHECKED(FAnimNextStateTreeRigVMConditionBase, ConditionFunctionName))
+	if (PropertyChangedEvent.GetMemberPropertyName() == GET_MEMBER_NAME_CHECKED(FAnimNextStateTreeRigVMConditionBase, CallFunctionInfo))
 	{
 		// Function Name selected has changed. Update param struct & result / event names used during execution.
 		if (FAnimNextStateTreeRigVMConditionInstanceData* InstanceData = InstanceDataView.GetMutablePtr<FAnimNextStateTreeRigVMConditionInstanceData>())
 		{
-			// @TODO: This relies on the function name being unique (Ex: In a workspace). For now that's okay. Later on however we will want to use a more robust function picker
-			auto GetRigVMFunctionHeader = [](FName InName)
-			{
-				FRigVMGraphFunctionHeader Result = FRigVMGraphFunctionHeader();
-
-				TMap<FAssetData, FRigVMGraphFunctionHeaderArray> FunctionExports;
-				UE::UAF::UncookedOnly::FUtils::GetExportedFunctionsFromAssetRegistry(UE::UAF::AnimNextPublicGraphFunctionsExportsRegistryTag, FunctionExports);
-				UE::UAF::UncookedOnly::FUtils::GetExportedFunctionsFromAssetRegistry(UE::UAF::ControlRigAssetPublicGraphFunctionsExportsRegistryTag, FunctionExports);
-
-				for (const auto& Export : FunctionExports.Array())
-				{
-					for (const FRigVMGraphFunctionHeader& FunctionHeader : Export.Value.Headers)
-					{
-						if (InName == FunctionHeader.Name)
-						{
-							Result = FunctionHeader;
-							return Result;
-						}
-					}
-				}
-
-				return Result;
-			};
-
-			RigVMFunctionHeader = GetRigVMFunctionHeader(ConditionFunctionName);
 			InstanceData->Parameters.Reset();
 			ResultIndex = INDEX_NONE;
 
 			TArray<FPropertyBagPropertyDesc> PropertyDescs;
-			PropertyDescs.Reserve(RigVMFunctionHeader.Arguments.Num());
-			for (const FRigVMGraphFunctionArgument& Argument : RigVMFunctionHeader.Arguments)
+			PropertyDescs.Reserve(CallFunctionInfo.FunctionHeader.Arguments.Num());
+			for (const FRigVMGraphFunctionArgument& Argument : CallFunctionInfo.FunctionHeader.Arguments)
 			{
 				if (Argument.Direction == ERigVMPinDirection::Input || Argument.Direction == ERigVMPinDirection::Output)
 				{
@@ -134,7 +109,7 @@ void FAnimNextStateTreeRigVMConditionBase::PostEditNodeChangeChainProperty(const
 				}
 			}
 
-			// TODO: Constrain functions we can pick to those returning only a bool
+			// @TODO: Constrain functions we can pick to those returning only a bool
 			ensure(ResultIndex != INDEX_NONE);
 
 			InstanceData->Parameters.AddProperties(PropertyDescs);
@@ -150,13 +125,33 @@ void FAnimNextStateTreeRigVMConditionBase::GetProgrammaticFunctionHeaders(FAnimN
 	
 	StateName = State->Name;
 	NodeId = Desc.ID;
-	InternalEventName = FName(FUtils::MakeFunctionWrapperEventName(RigVMFunctionHeader.Name));
 
 	FAnimNextProgrammaticFunctionHeader AnimNextFunctionHeader = {};
-	AnimNextFunctionHeader.FunctionHeader = RigVMFunctionHeader;
+	AnimNextFunctionHeader.FunctionHeader = CallFunctionInfo.FunctionHeader;
 	OutCompileContext.AddUniqueFunctionHeader(AnimNextFunctionHeader);
 }
 
 #endif // WITH_EDITOR
 
+void FAnimNextStateTreeRigVMConditionBase::PostSerialize(const FArchive& Ar)
+{
+#if WITH_EDITORONLY_DATA
+	PRAGMA_DISABLE_DEPRECATION_WARNINGS
+	if (Ar.IsLoading())
+	{
+		if (RigVMFunctionHeader.IsValid() && !ConditionFunctionName.IsNone() && !InternalEventName.IsNone())
+		{
+			CallFunctionInfo.FunctionHeader = RigVMFunctionHeader;
+			CallFunctionInfo.Function = ConditionFunctionName;
+			CallFunctionInfo.FunctionEvent = InternalEventName;
+
+			RigVMFunctionHeader = {};
+			ConditionFunctionName = {};
+			InternalEventName = {};
+		}
+	}
+	PRAGMA_ENABLE_DEPRECATION_WARNINGS
+#endif // WITH_EDITORONLY_DATA
+}
+
 #undef LOCTEXT_NAMESPACE
diff --git a/Engine/Plugins/Experimental/UAF/UAFStateTree/Source/UAFStateTree/Private/Tasks/AnimNextStateTreeRigVMTaskBase.cpp b/Engine/Plugins/Experimental/UAF/UAFStateTree/Source/UAFStateTree/Private/Tasks/AnimNextStateTreeRigVMTaskBase.cpp
--- a/Engine/Plugins/Experimental/UAF/UAFStateTree/Source/UAFStateTree/Private/Tasks/AnimNextStateTreeRigVMTaskBase.cpp
+++ b/Engine/Plugins/Experimental/UAF/UAFStateTree/Source/UAFStateTree/Private/Tasks/AnimNextStateTreeRigVMTaskBase.cpp
@@ -50,7 +50,7 @@ EStateTreeRunStatus FAnimNextStateTreeRigVMTaskBase::EnterState(FStateTreeExecut
 			if (!InstanceData->FunctionHandle.IsValid())
 #endif
 			{
-				InstanceData->FunctionHandle = RigVMAsset->GetFunctionHandle(InternalEventName);
+				InstanceData->FunctionHandle = RigVMAsset->GetFunctionHandle(CallFunctionInfo.FunctionEvent);
 			}
 
 			// See below, for now ignore index as we fire and forget on tasks
@@ -63,7 +63,7 @@ EStateTreeRunStatus FAnimNextStateTreeRigVMTaskBase::EnterState(FStateTreeExecut
 			FRigVMExtendedExecuteContext& ExtendedExecuteContext = RigVMComponent.GetExtendedExecuteContext();
 			FAnimNextExecuteContext& AnimNextContext = ExtendedExecuteContext.GetPublicDataSafe<FAnimNextExecuteContext>();
 
-			FAnimNextGraphContextData ContextData(GraphInstance, InternalEventName, AnimNextContext.GetDeltaTime<float>());
+			FAnimNextGraphContextData ContextData(GraphInstance, CallFunctionInfo.FunctionEvent, AnimNextContext.GetDeltaTime<float>());
 			UE::UAF::FScopedExecuteContextData ContextDataScope(AnimNextContext, ContextData);
 
 			// @TODO: Add option to execute per tick?
@@ -88,42 +88,17 @@ FText FAnimNextStateTreeRigVMTaskBase::GetDescription(const FGuid& ID, FStateTre
 
 void FAnimNextStateTreeRigVMTaskBase::PostEditNodeChangeChainProperty(const FPropertyChangedChainEvent& PropertyChangedEvent, FStateTreeDataView InstanceDataView)
 {
-	if (PropertyChangedEvent.GetMemberPropertyName() == GET_MEMBER_NAME_CHECKED(FAnimNextStateTreeRigVMTaskBase, TaskFunctionName))
+	if (PropertyChangedEvent.GetMemberPropertyName() == GET_MEMBER_NAME_CHECKED(FAnimNextStateTreeRigVMTaskBase, CallFunctionInfo))
 	{
-		// Function Name selected has changed. Update param struct & result / event names used during execution.
+		// Function selected has changed. Update param struct & result index.
 		if (FAnimNextStateTreeRigVMTaskInstanceData* InstanceData = InstanceDataView.GetMutablePtr<FAnimNextStateTreeRigVMTaskInstanceData>())
 		{
-			// @TODO: This relies on the function name being unique (Ex: In a workspace). For now that's okay. Later on however we will want to use a more robust function picker
-			auto GetRigVMFunctionHeader = [](FName InName)
-			{
-				FRigVMGraphFunctionHeader Result = FRigVMGraphFunctionHeader();
-
-				TMap<FAssetData, FRigVMGraphFunctionHeaderArray> FunctionExports;
-				UE::UAF::UncookedOnly::FUtils::GetExportedFunctionsFromAssetRegistry(UE::UAF::AnimNextPublicGraphFunctionsExportsRegistryTag, FunctionExports);
-				UE::UAF::UncookedOnly::FUtils::GetExportedFunctionsFromAssetRegistry(UE::UAF::ControlRigAssetPublicGraphFunctionsExportsRegistryTag, FunctionExports);
-
-				for (const auto& Export : FunctionExports.Array())
-				{
-					for (const FRigVMGraphFunctionHeader& FunctionHeader : Export.Value.Headers)
-					{
-						if (InName == FunctionHeader.Name)
-						{
-							Result = FunctionHeader;
-							return Result;
-						}
-					}
-				}
-
-				return Result;
-			};
-
-			RigVMFunctionHeader = GetRigVMFunctionHeader(TaskFunctionName);
 			InstanceData->Parameters.Reset();
 			ResultIndex = INDEX_NONE;
 
 			TArray<FPropertyBagPropertyDesc> PropertyDescs;
-			PropertyDescs.Reserve(RigVMFunctionHeader.Arguments.Num());
-			for (const FRigVMGraphFunctionArgument& Argument : RigVMFunctionHeader.Arguments)
+			PropertyDescs.Reserve(CallFunctionInfo.FunctionHeader.Arguments.Num());
+			for (const FRigVMGraphFunctionArgument& Argument : CallFunctionInfo.FunctionHeader.Arguments)
 			{
 				if (Argument.Direction == ERigVMPinDirection::Input || Argument.Direction == ERigVMPinDirection::Output)
 				{
@@ -152,13 +127,33 @@ void FAnimNextStateTreeRigVMTaskBase::GetProgrammaticFunctionHeaders(FAnimNextSt
 	
 	StateName = State->Name;
 	NodeId = Desc.ID;
-	InternalEventName = FName(FUtils::MakeFunctionWrapperEventName(RigVMFunctionHeader.Name));
 
 	FAnimNextProgrammaticFunctionHeader AnimNextFunctionHeader = {};
-	AnimNextFunctionHeader.FunctionHeader = RigVMFunctionHeader;
+	AnimNextFunctionHeader.FunctionHeader = CallFunctionInfo.FunctionHeader;
 	OutCompileContext.AddUniqueFunctionHeader(AnimNextFunctionHeader);
 }
 
 #endif // WITH_EDITOR
 
+void FAnimNextStateTreeRigVMTaskBase::PostSerialize(const FArchive& Ar)
+{
+#if WITH_EDITORONLY_DATA
+	PRAGMA_DISABLE_DEPRECATION_WARNINGS
+	if (Ar.IsLoading())
+	{
+		if (RigVMFunctionHeader.IsValid() && !TaskFunctionName.IsNone() && !InternalEventName.IsNone())
+		{
+			CallFunctionInfo.FunctionHeader = RigVMFunctionHeader;
+			CallFunctionInfo.Function = TaskFunctionName;
+			CallFunctionInfo.FunctionEvent = InternalEventName;
+
+			RigVMFunctionHeader = {};
+			TaskFunctionName = {};
+			InternalEventName = {};
+		}
+	}
+	PRAGMA_ENABLE_DEPRECATION_WARNINGS
+#endif // WITH_EDITORONLY_DATA
+}
+
 #undef LOCTEXT_NAMESPACE
diff --git a/Engine/Plugins/Experimental/UAF/UAFStateTree/Source/UAFStateTreeUncookedOnly/Private/AnimNextStateTreeFunctionLibraryHelper.cpp b/Engine/Plugins/Experimental/UAF/UAFStateTree/Source/UAFStateTreeUncookedOnly/Private/AnimNextStateTreeFunctionLibraryHelper.cpp
deleted file mode 100644
diff --git a/Engine/Plugins/Experimental/UAF/UAFStateTree/Source/UAFStateTreeUncookedOnly/Public/AnimNextStateTreeFunctionLibraryHelper.h b/Engine/Plugins/Experimental/UAF/UAFStateTree/Source/UAFStateTreeUncookedOnly/Public/AnimNextStateTreeFunctionLibraryHelper.h
deleted file mode 100644
