UAF: Add more logging to areas that could be important to users

#Synced-CL 48757676

diff --git a/Engine/Plugins/Experimental/MoverAnimNext/Source/MoverAnimNext/Private/Graph/RigUnit_GenerateMoverTrajectory.cpp b/Engine/Plugins/Experimental/MoverAnimNext/Source/MoverAnimNext/Private/Graph/RigUnit_GenerateMoverTrajectory.cpp
--- a/Engine/Plugins/Experimental/MoverAnimNext/Source/MoverAnimNext/Private/Graph/RigUnit_GenerateMoverTrajectory.cpp
+++ b/Engine/Plugins/Experimental/MoverAnimNext/Source/MoverAnimNext/Private/Graph/RigUnit_GenerateMoverTrajectory.cpp
@@ -6,6 +6,7 @@
 #include "MoverComponent.h"
 #include "MoverPoseSearchTrajectoryPredictor.h"
 #include "PoseSearch/PoseSearchTrajectoryLibrary.h"
+#include "UAFLogging.h"
 
 FRigUnit_GenerateMoverTrajectory_Execute()
 {
@@ -30,7 +31,7 @@ FRigUnit_GenerateMoverTrajectory_Execute()
 	if (MoverComponent == nullptr)
 	{
 		// Cannot execute, abort
-		// Toreview: error / warning logging necessary?
+		UAF_RIGUNIT_LOG(Warning, TEXT("Could not generate trajectory - Mover component is not valid."));
 		return;
 	}
 
diff --git a/Engine/Plugins/Experimental/MoverAnimNext/Source/MoverAnimNext/Private/Graph/RigUnit_MoverToggleRootMotion.cpp b/Engine/Plugins/Experimental/MoverAnimNext/Source/MoverAnimNext/Private/Graph/RigUnit_MoverToggleRootMotion.cpp
--- a/Engine/Plugins/Experimental/MoverAnimNext/Source/MoverAnimNext/Private/Graph/RigUnit_MoverToggleRootMotion.cpp
+++ b/Engine/Plugins/Experimental/MoverAnimNext/Source/MoverAnimNext/Private/Graph/RigUnit_MoverToggleRootMotion.cpp
@@ -5,6 +5,7 @@
 #include "DefaultMovementSet/LayeredMoves/RootMotionAttributeLayeredMove.h"
 #include "MoverComponent.h"
 #include "Module/AnimNextModuleInstance.h"
+#include "UAFLogging.h"
 
 FRigUnit_MoverToggleRootMotion_Execute()
 {
@@ -13,6 +14,7 @@ FRigUnit_MoverToggleRootMotion_Execute()
 
 	if (!MoverComponent || !MoverComponent->HasBeenInitialized())
 	{
+		UAF_RIGUNIT_LOG(Warning, TEXT("Could not toggle root motion - No valid mover component provided, or component has not been initialized yet."));
 		return;
 	}
 
diff --git a/Engine/Plugins/Experimental/UAF/UAF/Source/UAF/Internal/AnimNextExecuteContext.h b/Engine/Plugins/Experimental/UAF/UAF/Source/UAF/Internal/AnimNextExecuteContext.h
--- a/Engine/Plugins/Experimental/UAF/UAF/Source/UAF/Internal/AnimNextExecuteContext.h
+++ b/Engine/Plugins/Experimental/UAF/UAF/Source/UAF/Internal/AnimNextExecuteContext.h
@@ -5,6 +5,8 @@
 #include "CoreMinimal.h"
 #include "RigVMCore/RigVMExecuteContext.h"
 #include "Module/AnimNextModuleContextData.h"
+#include "Logging/StructuredLog.h"
+
 #include "AnimNextExecuteContext.generated.h"
 
 struct FUAFAssetInstance;
@@ -31,6 +33,29 @@ struct FAnimNextExecuteContext : public FRigVMExecuteContext
 		ContextData = OtherContext->ContextData;
 	}
 
+	virtual void Log(const FRigVMLogSettings& InLogSettings, const FString& InMessage) const override
+	{
+		if (RuntimeSettings.LogFunction.IsValid())
+		{
+			(*RuntimeSettings.LogFunction)(InLogSettings, this, InMessage);
+		}
+		else
+		{
+			if (InLogSettings.Severity == EMessageSeverity::Error)
+			{
+				UE_LOGFMT(LogAnimation, Error, "{FunctionName}: {Message} [Actor: {ActorPath}]", GetFunctionName(), InMessage, FSoftObjectPath(GetOwningActor()));
+			}
+			else if (InLogSettings.Severity == EMessageSeverity::Warning)
+			{
+				UE_LOGFMT(LogAnimation, Warning, "{FunctionName}: {Message} [Actor: {ActorPath}]", GetFunctionName(), InMessage, FSoftObjectPath(GetOwningActor()));
+			}
+			else
+			{
+				UE_LOGFMT(LogAnimation, Display, "{FunctionName}: {Message} [Actor: {ActorPath}]", GetFunctionName(), InMessage, FSoftObjectPath(GetOwningActor()));
+			}
+		}
+	}
+
 	// Get the context data as the specified type. This will assert if the type differs from the last call to SetContextData.
 	template<typename ContextType>
 	const ContextType& GetContextData() const
diff --git a/Engine/Plugins/Experimental/UAF/UAF/Source/UAF/Internal/UAFLogging.h b/Engine/Plugins/Experimental/UAF/UAF/Source/UAF/Internal/UAFLogging.h
new file mode 100644
--- /dev/null
+++ b/Engine/Plugins/Experimental/UAF/UAF/Source/UAF/Internal/UAFLogging.h
@@ -0,0 +1,30 @@
+// Copyright Epic Games, Inc. All Rights Reserved.
+
+#pragma once
+#include "Logging/LogCategory.h"
+#include "Logging/LogVerbosity.h"
+
+// Calls Log on the existing ExecuteContext within the current RigVMStruct scope. 
+// This will either call the bound LogFunction on the ExecuteContext or fallback to a regular UE_LOG and automatically prefixes the function name to the message
+#define UAF_RIGUNIT_LOG(Severity, Format, ...) \
+ExecuteContext.Log(EMessageSeverity::Severity, FString::Printf((Format), ##__VA_ARGS__))
+
+#if NO_LOGGING
+	// If logging is disabled, we follow UE_LOG's behavior and only log fatal messages
+	#define UAF_TRAIT_LOG(Verbosity, Format, ...) \
+	{\
+		if constexpr ((::ELogVerbosity::Verbosity & ::ELogVerbosity::VerbosityMask) == ::ELogVerbosity::Fatal) \
+		{ \
+			FString TraitName; if(Binding.GetTrait() != nullptr) { TraitName = Binding.GetTrait()->GetTraitName(); } else {TraitName = TEXT("Invalid Trait");} \
+			Context.Log(ELogVerbosity::Verbosity, FString::Printf((Format), ##__VA_ARGS__), TraitName); \
+		} \
+	}
+#else
+	// Calls Context.Log within the current trait scope
+	// This will automatically append the context's root graph name and the name of the owning object that runs the UAF setup, if available
+	#define UAF_TRAIT_LOG(Verbosity, Format, ...) \
+	FString TraitName; if(Binding.GetTrait() != nullptr) { TraitName = Binding.GetTrait()->GetTraitName(); } else {TraitName = TEXT("Invalid Trait");} \
+	Context.Log(ELogVerbosity::Verbosity, FString::Printf((Format), ##__VA_ARGS__), TraitName);
+
+#endif // NO_LOGGING
+
diff --git a/Engine/Plugins/Experimental/UAF/UAF/Source/UAF/Private/Component/AnimNextComponent.cpp b/Engine/Plugins/Experimental/UAF/UAF/Source/UAF/Private/Component/AnimNextComponent.cpp
--- a/Engine/Plugins/Experimental/UAF/UAF/Source/UAF/Private/Component/AnimNextComponent.cpp
+++ b/Engine/Plugins/Experimental/UAF/UAF/Source/UAF/Private/Component/AnimNextComponent.cpp
@@ -406,7 +406,7 @@ void UAnimNextComponent::AddModuleEventSubsequent(FName InEventName, UAnimNextCo
 	}
 	else if (OtherAnimNextComponent == this)
 	{
-		UE_LOG(LogAnimation, Warning, TEXT("UAFComponent::AddModuleEventPrerequisite called using the same component"));
+		UE_LOG(LogAnimation, Warning, TEXT("UAFComponent::AddModuleEventSubsequent called using the same component"));
 	}
 	else
 	{
@@ -418,11 +418,11 @@ void UAnimNextComponent::RemoveModuleEventPrerequisite(FName InEventName, UAnimN
 {
 	if (!OtherAnimNextComponent)
 	{
-		UE_LOG(LogAnimation, Warning, TEXT("UAnimNextComponent::RemoveModuleEventPrerequisite called with null OtherAnimNextComponent"));
+		UE_LOG(LogAnimation, Warning, TEXT("UAFComponent::RemoveModuleEventPrerequisite called with null OtherAnimNextComponent"));
 	}
 	else if (OtherAnimNextComponent == this)
 	{
-		UE_LOG(LogAnimation, Warning, TEXT("UAnimNextComponent::RemoveModuleEventPrerequisite called using the same component"));
+		UE_LOG(LogAnimation, Warning, TEXT("UAFComponent::RemoveModuleEventPrerequisite called using the same component"));
 	}
 	else
 	{
@@ -434,11 +434,11 @@ void UAnimNextComponent::RemoveModuleEventSubsequent(FName InEventName, UAnimNex
 {
 	if (!OtherAnimNextComponent)
 	{
-		UE_LOG(LogAnimation, Warning, TEXT("UAnimNextComponent::RemoveModuleEventSubsequent called with null OtherAnimNextComponent"));
+		UE_LOG(LogAnimation, Warning, TEXT("UAFComponent::RemoveModuleEventSubsequent called with null OtherAnimNextComponent"));
 	}
 	else if (OtherAnimNextComponent == this)
 	{
-		UE_LOG(LogAnimation, Warning, TEXT("UAnimNextComponent::RemoveModuleEventSubsequent called using the same component"));
+		UE_LOG(LogAnimation, Warning, TEXT("UAFComponent::RemoveModuleEventSubsequent called using the same component"));
 	}
 	else
 	{
diff --git a/Engine/Plugins/Experimental/UAF/UAF/Source/UAF/Private/Graph/RigUnit_AnimNextRemapPose.cpp b/Engine/Plugins/Experimental/UAF/UAF/Source/UAF/Private/Graph/RigUnit_AnimNextRemapPose.cpp
--- a/Engine/Plugins/Experimental/UAF/UAF/Source/UAF/Private/Graph/RigUnit_AnimNextRemapPose.cpp
+++ b/Engine/Plugins/Experimental/UAF/UAF/Source/UAF/Private/Graph/RigUnit_AnimNextRemapPose.cpp
@@ -5,6 +5,7 @@
 #include "GenerationTools.h"
 #include "Components/SkeletalMeshComponent.h"
 #include "Engine/SkeletalMesh.h"
+#include "UAFLogging.h"
 
 #include UE_INLINE_GENERATED_CPP_BY_NAME(RigUnit_AnimNextRemapPose)
 DEFINE_STAT(STAT_AnimNext_RigUnit_RemapPose);
@@ -17,11 +18,13 @@ FRigUnit_AnimNextRemapPose_Execute()
 
 	if(!TargetAnimGraphRefPose.ReferencePose.IsValid())
 	{
+		UAF_RIGUNIT_LOG(Warning, TEXT("Could not remap pose - Target Ref Pose is not valid."));
 		return;
 	}
 	
 	if(!Pose.LODPose.IsValid())
 	{
+		UAF_RIGUNIT_LOG(Warning, TEXT("Could not remap pose - LOD Pose is not valid."));
 		return;
 	}
 
diff --git a/Engine/Plugins/Experimental/UAF/UAF/Source/UAF/Private/Graph/RigUnit_AnimNextRetargetPose.cpp b/Engine/Plugins/Experimental/UAF/UAF/Source/UAF/Private/Graph/RigUnit_AnimNextRetargetPose.cpp
--- a/Engine/Plugins/Experimental/UAF/UAF/Source/UAF/Private/Graph/RigUnit_AnimNextRetargetPose.cpp
+++ b/Engine/Plugins/Experimental/UAF/UAF/Source/UAF/Private/Graph/RigUnit_AnimNextRetargetPose.cpp
@@ -10,6 +10,7 @@
 #include "Animation/SkeletonRemappingRegistry.h"
 #include "Animation/Skeleton.h"
 #include "Animation/SkeletonRemapping.h"
+#include "UAFLogging.h"
 
 #include UE_INLINE_GENERATED_CPP_BY_NAME(RigUnit_AnimNextRetargetPose)
 DEFINE_STAT(STAT_AnimNext_RigUnit_RetargetPose);
@@ -20,8 +21,15 @@ FRigUnit_AnimNextRetargetPose_Execute()
 
 	using namespace UE::UAF;
 
-	if(!TargetAnimGraphRefPose.ReferencePose.IsValid() || !Pose.LODPose.IsValid())
+	if(!TargetAnimGraphRefPose.ReferencePose.IsValid())
 	{
+		UAF_RIGUNIT_LOG(Warning, TEXT("Could not retarget pose - Target Ref Pose is not valid."));
+		return;
+	}
+
+	if (!Pose.LODPose.IsValid())
+	{
+		UAF_RIGUNIT_LOG(Warning, TEXT("Could not retarget pose - LOD Pose is not valid."));
 		return;
 	}
 
diff --git a/Engine/Plugins/Experimental/UAF/UAF/Source/UAF/Private/Graph/RigUnit_AnimNextWriteSkeletalMeshComponentPose.cpp b/Engine/Plugins/Experimental/UAF/UAF/Source/UAF/Private/Graph/RigUnit_AnimNextWriteSkeletalMeshComponentPose.cpp
--- a/Engine/Plugins/Experimental/UAF/UAF/Source/UAF/Private/Graph/RigUnit_AnimNextWriteSkeletalMeshComponentPose.cpp
+++ b/Engine/Plugins/Experimental/UAF/UAF/Source/UAF/Private/Graph/RigUnit_AnimNextWriteSkeletalMeshComponentPose.cpp
@@ -12,6 +12,7 @@
 #include "Logging/StructuredLog.h"
 #include "Module/AnimNextModuleInstance.h"
 #include "Module/AnimNextSkeletalMeshComponentReferenceComponent.h"
+#include "UAFLogging.h"
 
 #include UE_INLINE_GENERATED_CPP_BY_NAME(RigUnit_AnimNextWriteSkeletalMeshComponentPose)
 
@@ -37,25 +38,29 @@ FRigUnit_AnimNextWriteSkeletalMeshComponentPose_Execute()
 
 	if(OutputComponent == nullptr)
 	{
+		UAF_RIGUNIT_LOG(Warning, TEXT("Could not write to skeletal mesh component - No valid skeletal mesh component found."));
 		return;
 	}
 
 	if(!Pose.LODPose.IsValid())
 	{
+		UAF_RIGUNIT_LOG(Warning, TEXT("Could not write to skeletal mesh component - LOD Pose is not valid."));
 		return;
 	}
 
 	USkeletalMesh* SkeletalMesh = OutputComponent->GetSkeletalMeshAsset();
 	if(SkeletalMesh == nullptr)
 	{
+		UAF_RIGUNIT_LOG(Warning, TEXT("Could not write to skeletal mesh component - No valid SkeletalMesh on [SK: %s][Owner: %s]"),
+			*OutputComponent->GetName(), *GetNameSafe(OutputComponent->GetOuter()));
 		return;
 	}
 
 	// We cannot write to the skeletal mesh component if it is using anim BP
 	if(OutputComponent->bEnableAnimation)
 	{
-		UE_LOGFMT(LogAnimation, Warning, "UAF: Could not write to skeletal mesh component - bEnableAnimation is true [SK: {SkeletalMesh}][Owner: {Owner}]",
-			*SkeletalMesh->GetName(), OutputComponent->GetOuter() != nullptr ? *OutputComponent->GetOuter()->GetName() : TEXT("*NO OWNER*"));
+		UAF_RIGUNIT_LOG(Warning, TEXT("Could not write to skeletal mesh component - bEnableAnimation is true[SK: %s][Owner: %s]"),
+			*OutputComponent->GetName(), *GetNameSafe(OutputComponent->GetOuter()));
 		return;
 	}
 
@@ -66,8 +71,9 @@ FRigUnit_AnimNextWriteSkeletalMeshComponentPose_Execute()
 	// You should remap the pose explicitly if this is what you need
 	if (&RefPose != &Pose.LODPose.GetRefPose())
 	{
-		UE_LOGFMT(LogAnimation, Warning, "UAF: Could not write to skeletal mesh component - The input pose and the skeletal mesh component use different reference poses [SK: {SkeletalMesh}][Owner: {Owner}]",
-			*SkeletalMesh->GetName(), OutputComponent->GetOuter() != nullptr ? *OutputComponent->GetOuter()->GetName() : TEXT("*NO OWNER*"));
+		UAF_RIGUNIT_LOG(Warning, TEXT("Could not write to skeletal mesh component - The input pose and the skeletal mesh component use different reference poses [SK: %s][Owner: %s]"),
+			*SkeletalMesh->GetName(), *GetNameSafe(OutputComponent->GetOuter()));
+
 		return;
 	}
 
diff --git a/Engine/Plugins/Experimental/UAF/UAF/Source/UAF/Private/Graph/RigUnit_GetActorTransform.cpp b/Engine/Plugins/Experimental/UAF/UAF/Source/UAF/Private/Graph/RigUnit_GetActorTransform.cpp
--- a/Engine/Plugins/Experimental/UAF/UAF/Source/UAF/Private/Graph/RigUnit_GetActorTransform.cpp
+++ b/Engine/Plugins/Experimental/UAF/UAF/Source/UAF/Private/Graph/RigUnit_GetActorTransform.cpp
@@ -4,6 +4,7 @@
 
 #include "Component/AnimNextComponent.h"
 #include "Module/AnimNextModuleInstance.h"
+#include "UAFLogging.h"
 
 #include UE_INLINE_GENERATED_CPP_BY_NAME(RigUnit_GetActorTransform)
 
@@ -23,18 +24,21 @@ void FAnimNextActorTransformComponent::OnBindToInstance()
 			UAnimNextComponent* Component = WeakAnimNextComponent.Get();
 			if(Component == nullptr)
 			{
+				UE_LOG(LogAnimation, Warning, TEXT("FAnimNextActorTransformComponent: UAF Component is invalid."));
 				return;
 			}
 
 			AActor* Owner = Component->GetOwner();
 			if(Owner == nullptr)
 			{
+				UE_LOG(LogAnimation, Warning, TEXT("FAnimNextActorTransformComponent: Owner of %s is invalid."), *Component->GetName());
 				return;
 			}
 
 			USceneComponent* SceneComponent = Owner->GetRootComponent();
 			if(SceneComponent == nullptr)
 			{
+				UE_LOG(LogAnimation, Warning, TEXT("FAnimNextActorTransformComponent: %s does not have a valid RootComponent"), *Owner->GetName());
 				return;
 			}
 
@@ -69,9 +73,10 @@ FRigUnit_GetActorTransform_Execute()
 	const FUAFAssetContextData& AssetContextData = ExecuteContext.GetContextData<FUAFAssetContextData>();
 	FUAFAssetInstance& AssetInstance = AssetContextData.GetInstance();
 	if (AssetInstance.GetRootInstance() == nullptr)
-        {
+    {
+		UAF_RIGUNIT_LOG(Warning, TEXT("Could not get actor transform - Root module instance is not valid."));
 		return;
-        }
+    }
 	const FAnimNextActorTransformComponent& TransformComponent = AssetInstance.GetRootInstance()->GetComponent<FAnimNextActorTransformComponent>();
 	Transform = TransformComponent.GetActorTransform();
 }
diff --git a/Engine/Plugins/Experimental/UAF/UAF/Source/UAF/Private/Graph/RigUnit_GetSocketTransform.cpp b/Engine/Plugins/Experimental/UAF/UAF/Source/UAF/Private/Graph/RigUnit_GetSocketTransform.cpp
--- a/Engine/Plugins/Experimental/UAF/UAF/Source/UAF/Private/Graph/RigUnit_GetSocketTransform.cpp
+++ b/Engine/Plugins/Experimental/UAF/UAF/Source/UAF/Private/Graph/RigUnit_GetSocketTransform.cpp
@@ -1,6 +1,7 @@
 // Copyright Epic Games, Inc. All Rights Reserved.
 
 #include "RigUnit_GetSocketTransform.h"
+#include "UAFLogging.h"
 
 #include UE_INLINE_GENERATED_CPP_BY_NAME(RigUnit_GetSocketTransform)
 
@@ -10,4 +11,8 @@ FRigUnit_GetSocketTransform_Execute()
 	{
 		Result = SceneComponent->GetSocketTransform(SocketName, TransformSpace);
 	}
+	else
+	{
+		UAF_RIGUNIT_LOG(Warning, TEXT("Could not get transform for socket [%s] - SceneComponent is not valid."), *SocketName.ToString());
+	}
 }
diff --git a/Engine/Plugins/Experimental/UAF/UAF/Source/UAF/Private/Graph/RigUnit_MakeReferencePoseFromSkeletalMeshComponent.cpp b/Engine/Plugins/Experimental/UAF/UAF/Source/UAF/Private/Graph/RigUnit_MakeReferencePoseFromSkeletalMeshComponent.cpp
--- a/Engine/Plugins/Experimental/UAF/UAF/Source/UAF/Private/Graph/RigUnit_MakeReferencePoseFromSkeletalMeshComponent.cpp
+++ b/Engine/Plugins/Experimental/UAF/UAF/Source/UAF/Private/Graph/RigUnit_MakeReferencePoseFromSkeletalMeshComponent.cpp
@@ -8,6 +8,7 @@
 #include "Engine/SkeletalMesh.h"
 #include "Module/AnimNextModuleInstance.h"
 #include "Module/AnimNextSkeletalMeshComponentReferenceComponent.h"
+#include "UAFLogging.h"
 
 #include UE_INLINE_GENERATED_CPP_BY_NAME(RigUnit_MakeReferencePoseFromSkeletalMeshComponent)
 
@@ -32,6 +33,7 @@ FRigUnit_MakeReferencePoseFromSkeletalMeshComponent_Execute()
 
 	if(InputComponent == nullptr)
 	{
+		UAF_RIGUNIT_LOG(Warning, TEXT("Could not make ref pose - Skeletal mesh component is not valid."));
 		return;
 	}
 
diff --git a/Engine/Plugins/Experimental/UAF/UAF/Source/UAF/Private/Module/AnimNextActorComponentReferenceComponent.cpp b/Engine/Plugins/Experimental/UAF/UAF/Source/UAF/Private/Module/AnimNextActorComponentReferenceComponent.cpp
--- a/Engine/Plugins/Experimental/UAF/UAF/Source/UAF/Private/Module/AnimNextActorComponentReferenceComponent.cpp
+++ b/Engine/Plugins/Experimental/UAF/UAF/Source/UAF/Private/Module/AnimNextActorComponentReferenceComponent.cpp
@@ -15,6 +15,7 @@ void FAnimNextActorComponentReferenceComponent::OnInitializeHelper(UScriptStruct
 	FAnimNextModuleInstance* ModuleInstance = GetModuleInstancePtr();
 	if (ModuleInstance == nullptr)
 	{
+		UE_LOG(LogAnimation, Warning, TEXT("FAnimNextActorComponentReferenceComponent: Could not retrieve component reference of type [%s] - Module Instance is not valid."), *GetNameSafe(ComponentType));
 		return;
 	}
 
@@ -26,12 +27,14 @@ void FAnimNextActorComponentReferenceComponent::OnInitializeHelper(UScriptStruct
 			UAnimNextComponent* PinnedAnimNextComponent = WeakAnimNextComponent.Get();
 			if(PinnedAnimNextComponent == nullptr)
 			{
+				UE_LOG(LogAnimation, Warning, TEXT("FAnimNextActorComponentReferenceComponent: Could not retrieve component reference of type [%s] - UAF Component is not valid."), *GetNameSafe(ComponentType));
 				return;
 			}
 
 			AActor* Owner = PinnedAnimNextComponent->GetOwner();
 			if(Owner == nullptr)
 			{
+				UE_LOG(LogAnimation, Warning, TEXT("FAnimNextActorComponentReferenceComponent: Could not retrieve component reference of type [%s] - Owner of UAF Component is not valid."), *GetNameSafe(ComponentType));
 				return;
 			}
 
diff --git a/Engine/Plugins/Experimental/UAF/UAF/Source/UAF/Private/Module/AnimNextModuleInstance.cpp b/Engine/Plugins/Experimental/UAF/UAF/Source/UAF/Private/Module/AnimNextModuleInstance.cpp
--- a/Engine/Plugins/Experimental/UAF/UAF/Source/UAF/Private/Module/AnimNextModuleInstance.cpp
+++ b/Engine/Plugins/Experimental/UAF/UAF/Source/UAF/Private/Module/AnimNextModuleInstance.cpp
@@ -590,6 +590,8 @@ void FAnimNextModuleInstance::RunScriptEvent(const UE::UAF::FScriptEvent& InEven
 {
 	if (!InEvent.IsCallable())
 	{
+		UE_LOG(LogAnimation, Warning, TEXT("FAnimNextModuleInstance::RunScriptEvent: Could not run script event [%s] - The event is not callable."),
+			*InEvent.GetEventName().ToString());
 		return;
 	}
 
diff --git a/Engine/Plugins/Experimental/UAF/UAF/Source/UAF/Private/Module/ModuleTickFunction.cpp b/Engine/Plugins/Experimental/UAF/UAF/Source/UAF/Private/Module/ModuleTickFunction.cpp
--- a/Engine/Plugins/Experimental/UAF/UAF/Source/UAF/Private/Module/ModuleTickFunction.cpp
+++ b/Engine/Plugins/Experimental/UAF/UAF/Source/UAF/Private/Module/ModuleTickFunction.cpp
@@ -312,6 +312,7 @@ void FModuleEventTickFunction::RunTaskOnGameThread(TUniqueFunction<void(void)>&&
 {
 	if (!InFunction)
 	{
+		UE_LOG(LogAnimation, Warning, TEXT("ModuleEventTickFunction::RunTaskOnGameThread: Could not run task - InFunction is not valid."));
 		return;
 	}
 
diff --git a/Engine/Plugins/Experimental/UAF/UAF/Source/UAF/Private/Module/RigUnit_AnimNextModuleEventDependencies.cpp b/Engine/Plugins/Experimental/UAF/UAF/Source/UAF/Private/Module/RigUnit_AnimNextModuleEventDependencies.cpp
--- a/Engine/Plugins/Experimental/UAF/UAF/Source/UAF/Private/Module/RigUnit_AnimNextModuleEventDependencies.cpp
+++ b/Engine/Plugins/Experimental/UAF/UAF/Source/UAF/Private/Module/RigUnit_AnimNextModuleEventDependencies.cpp
@@ -44,6 +44,7 @@ static void ApplyDependenciesHelper(const FAnimNextExecuteContext& InContext, co
 
 	if(Dependencies.Num() == 0)
 	{
+		UE_LOG(LogAnimation, Warning, TEXT("ApplyDependenciesHelper: Could not apply dependencies - No dependencies found."));
 		return;
 	}
 
@@ -60,6 +61,7 @@ static void ApplyDependenciesHelper(const FAnimNextExecuteContext& InContext, co
 
 	if(!bTickFunctionFound)
 	{
+		UE_LOG(LogAnimation, Warning, TEXT("ApplyDependenciesHelper: Could not apply dependencies - No valid tick function found."));
 		return;
 	}
 
@@ -69,6 +71,7 @@ static void ApplyDependenciesHelper(const FAnimNextExecuteContext& InContext, co
 		UObject* Object = WeakObject.Get();
 		if(Object == nullptr)
 		{
+			UE_LOG(LogAnimation, Warning, TEXT("ApplyDependenciesHelper: Could not apply dependencies - Module Object is not valid."));
 			return;
 		}
 
diff --git a/Engine/Plugins/Experimental/UAF/UAF/Source/UAF/Private/Module/RigVMTrait_ModuleEventDependency_ActorComponentClassPrimaryTickFunction.cpp b/Engine/Plugins/Experimental/UAF/UAF/Source/UAF/Private/Module/RigVMTrait_ModuleEventDependency_ActorComponentClassPrimaryTickFunction.cpp
--- a/Engine/Plugins/Experimental/UAF/UAF/Source/UAF/Private/Module/RigVMTrait_ModuleEventDependency_ActorComponentClassPrimaryTickFunction.cpp
+++ b/Engine/Plugins/Experimental/UAF/UAF/Source/UAF/Private/Module/RigVMTrait_ModuleEventDependency_ActorComponentClassPrimaryTickFunction.cpp
@@ -16,18 +16,22 @@ void FRigVMTrait_ModuleEventDependency_ActorComponentClassPrimaryTickFunction::O
 	UActorComponent* AnimNextComponent = Cast<UActorComponent>(InContext.Object);
 	if(AnimNextComponent == nullptr)
 	{
+		UE_LOG(LogAnimation, Warning, TEXT("FRigVMTrait_ModuleEventDependency_ActorComponentClassPrimaryTickFunction: Could not add dependency for component of class [%s] - UAF Component is not valid."), *GetNameSafe(ComponentClass));
 		return;
 	}
 
 	AActor* Actor = AnimNextComponent->GetOwner();
 	if(Actor == nullptr)
 	{
+		UE_LOG(LogAnimation, Warning, TEXT("FRigVMTrait_ModuleEventDependency_ActorComponentClassPrimaryTickFunction: Could not add dependency for component of class [%s] - UAF Component has no valid owner."), *GetNameSafe(ComponentClass));
 		return;
 	}
 
 	UActorComponent* TargetComponent = Actor->GetComponentByClass(ComponentClass);
 	if(TargetComponent == nullptr)
 	{
+		UE_LOG(LogAnimation, Log, TEXT("FRigVMTrait_ModuleEventDependency_ActorComponentClassPrimaryTickFunction: Could not add dependency for component of class [%s] - [%s] does not have that component."),
+			*GetNameSafe(ComponentClass), *Actor->GetName());
 		return;
 	}
 
@@ -46,18 +50,22 @@ void FRigVMTrait_ModuleEventDependency_ActorComponentClassPrimaryTickFunction::O
 	UActorComponent* AnimNextComponent = Cast<UActorComponent>(InContext.Object);
 	if(AnimNextComponent == nullptr)
 	{
+		UE_LOG(LogAnimation, Warning, TEXT("FRigVMTrait_ModuleEventDependency_ActorComponentClassPrimaryTickFunction: Could not remove dependency for component of class [%s] - UAF Component is not valid."), *GetNameSafe(ComponentClass));
 		return;
 	}
 
 	AActor* Actor = AnimNextComponent->GetOwner();
 	if(Actor == nullptr)
 	{
+		UE_LOG(LogAnimation, Warning, TEXT("FRigVMTrait_ModuleEventDependency_ActorComponentClassPrimaryTickFunction: Could not remove dependency for component of class [%s] - UAF Component has no valid owner."), *GetNameSafe(ComponentClass));
 		return;
 	}
 
 	UActorComponent* TargetComponent = Actor->GetComponentByClass(ComponentClass);
 	if(TargetComponent == nullptr)
 	{
+		UE_LOG(LogAnimation, Log, TEXT("FRigVMTrait_ModuleEventDependency_ActorComponentClassPrimaryTickFunction: Could not remove dependency for component of class [%s] - [%s] does not have that component."),
+			*GetNameSafe(ComponentClass), *Actor->GetName());
 		return;
 	}
 
diff --git a/Engine/Plugins/Experimental/UAF/UAF/Source/UAF/Private/Module/RigVMTrait_ModuleEventDependency_ActorComponentPrimaryTickFunction.cpp b/Engine/Plugins/Experimental/UAF/UAF/Source/UAF/Private/Module/RigVMTrait_ModuleEventDependency_ActorComponentPrimaryTickFunction.cpp
--- a/Engine/Plugins/Experimental/UAF/UAF/Source/UAF/Private/Module/RigVMTrait_ModuleEventDependency_ActorComponentPrimaryTickFunction.cpp
+++ b/Engine/Plugins/Experimental/UAF/UAF/Source/UAF/Private/Module/RigVMTrait_ModuleEventDependency_ActorComponentPrimaryTickFunction.cpp
@@ -16,11 +16,13 @@ void FRigVMTrait_ModuleEventDependency_ActorComponentPrimaryTickFunction::OnAddD
 	UActorComponent* AnimNextComponent = Cast<UActorComponent>(InContext.Object);
 	if(AnimNextComponent == nullptr)
 	{
+		UE_LOG(LogAnimation, Warning, TEXT("FRigVMTrait_ModuleEventDependency_ActorComponentPrimaryTickFunction: Could not add dependency for component [%s] - UAF Component is not valid."), *GetNameSafe(Component));
 		return;
 	}
 
 	if(Component == nullptr)
 	{
+		UE_LOG(LogAnimation, Warning, TEXT("FRigVMTrait_ModuleEventDependency_ActorComponentPrimaryTickFunction: Could not add dependency - Component is not valid."));
 		return;
 	}
 
@@ -39,11 +41,13 @@ void FRigVMTrait_ModuleEventDependency_ActorComponentPrimaryTickFunction::OnRemo
 	UActorComponent* AnimNextComponent = Cast<UActorComponent>(InContext.Object);
 	if(AnimNextComponent == nullptr)
 	{
+		UE_LOG(LogAnimation, Warning, TEXT("FRigVMTrait_ModuleEventDependency_ActorComponentPrimaryTickFunction: Could not remove dependency for component [%s] - UAF Component is not valid."), *GetNameSafe(Component));
 		return;
 	}
 	
 	if(Component == nullptr)
 	{
+		UE_LOG(LogAnimation, Warning, TEXT("FRigVMTrait_ModuleEventDependency_ActorComponentPrimaryTickFunction: Could not remove dependency - Component is not valid."));
 		return;
 	}
 
diff --git a/Engine/Plugins/Experimental/UAF/UAF/Source/UAF/Private/Script/UAFRigVMComponent.cpp b/Engine/Plugins/Experimental/UAF/UAF/Source/UAF/Private/Script/UAFRigVMComponent.cpp
--- a/Engine/Plugins/Experimental/UAF/UAF/Source/UAF/Private/Script/UAFRigVMComponent.cpp
+++ b/Engine/Plugins/Experimental/UAF/UAF/Source/UAF/Private/Script/UAFRigVMComponent.cpp
@@ -19,6 +19,7 @@ void FUAFRigVMComponent::OnBindToInstance()
 	VM = Asset->RigVM;
 	if (VM == nullptr)
 	{
+		UE_LOG(LogAnimation, Warning, TEXT("FUAFRigVMComponent::OnBindToInstance: VM is not valid."));
 		return;
 	}
 
@@ -62,12 +63,14 @@ void FUAFRigVMComponent::StaticCallEvent(TStructView<FUAFScriptComponent> InScri
 
 	if(This.VM == nullptr)
 	{
+		UE_LOG(LogAnimation, Warning, TEXT("FUAFRigVMComponent::StaticCallEvent: Could not call event - RigVM is not valid."));
 		return;
 	}
 
 	FName EventName = InContextData.Get<FUAFScriptContextData>().GetEventName();
 	if(!This.VM->ContainsEntry(EventName))
 	{
+		UE_LOG(LogAnimation, Warning, TEXT("FUAFRigVMComponent::StaticCallEvent: Could not call event - VM does not contain event [%s]"), *EventName.ToString());
 		return;
 	}
 
diff --git a/Engine/Plugins/Experimental/UAF/UAF/Source/UAF/Private/UAFAssetInstance.cpp b/Engine/Plugins/Experimental/UAF/UAF/Source/UAF/Private/UAFAssetInstance.cpp
--- a/Engine/Plugins/Experimental/UAF/UAF/Source/UAF/Private/UAFAssetInstance.cpp
+++ b/Engine/Plugins/Experimental/UAF/UAF/Source/UAF/Private/UAFAssetInstance.cpp
@@ -67,6 +67,7 @@ bool FUAFAssetInstance::LayoutMatches(const FInstancedPropertyBag& InPropertyBag
 
 	if (Variables.NumInternalVariables != InPropertyBag.GetNumPropertiesInBag())
 	{
+		UE_LOG(LogAnimation, Warning, TEXT("FUAFAssetInstance::LayoutMatches: Number of internal variables does not match number of properties in PropertyBag."));
 		return false;
 	}
 
@@ -76,6 +77,7 @@ bool FUAFAssetInstance::LayoutMatches(const FInstancedPropertyBag& InPropertyBag
 	{
 		if (!Descs.IsValidIndex(InVariableIndex))
 		{
+			UE_LOG(LogAnimation, Warning, TEXT("FUAFAssetInstance::LayoutMatches: Variable Index %d is not a valid index into the PropertyBag's description."), InVariableIndex);
 			bMatches = false;
 			return false;
 		}
@@ -83,6 +85,7 @@ bool FUAFAssetInstance::LayoutMatches(const FInstancedPropertyBag& InPropertyBag
 		const FPropertyBagPropertyDesc& Desc = Descs[InVariableIndex];
 		if (Desc.Name != InName)
 		{
+			UE_LOG(LogAnimation, Warning, TEXT("FUAFAssetInstance::LayoutMatches: Variable name [%s] does not match the PropertyBag's entry name [%s]."), *InName.ToString(), *Desc.Name.ToString());
 			bMatches = false;
 			return false;
 		}
@@ -91,6 +94,7 @@ bool FUAFAssetInstance::LayoutMatches(const FInstancedPropertyBag& InPropertyBag
 			Desc.ContainerTypes.GetFirstContainerType() != InType.GetContainerType() ||
 			Desc.ValueTypeObject != InType.GetValueTypeObject())
 		{
+			UE_LOG(LogAnimation, Warning, TEXT("FUAFAssetInstance::LayoutMatches: Variable type does not match the type of the PropertyBag's entry."));
 			bMatches = false;
 			return false;
 		}
diff --git a/Engine/Plugins/Experimental/UAF/UAFAnimGraph/Source/UAFAnimGraph/Private/AnimNextAnimGraphModule.cpp b/Engine/Plugins/Experimental/UAF/UAFAnimGraph/Source/UAFAnimGraph/Private/AnimNextAnimGraphModule.cpp
--- a/Engine/Plugins/Experimental/UAF/UAFAnimGraph/Source/UAFAnimGraph/Private/AnimNextAnimGraphModule.cpp
+++ b/Engine/Plugins/Experimental/UAF/UAFAnimGraph/Source/UAFAnimGraph/Private/AnimNextAnimGraphModule.cpp
@@ -275,7 +275,7 @@ void FAnimNextAnimGraphModule::ListNodeTemplates(const TArray<FString>& Args)
 	const FNodeTemplateRegistry& NodeTemplateRegistry = FNodeTemplateRegistry::Get();
 	const FTraitRegistry& TraitRegistry = FTraitRegistry::Get();
 
-	UE_LOG(LogAnimation, Log, TEXT("===== AnimNext Node Templates ====="));
+	UE_LOG(LogAnimation, Log, TEXT("===== UAF Node Templates ====="));
 	UE_LOG(LogAnimation, Log, TEXT("Template Buffer Size: %u bytes"), NodeTemplateRegistry.TemplateBuffer.GetAllocatedSize());
 
 	for (auto It = NodeTemplateRegistry.TemplateUIDToHandleMap.CreateConstIterator(); It; ++It)
@@ -347,7 +347,7 @@ void FAnimNextAnimGraphModule::ListAnimationGraphs(const TArray<FString>& Args)
 	const FTraitRegistry& TraitRegistry = FTraitRegistry::Get();
 	const bool bDetailedOutput = true;
 
-	UE_LOG(LogAnimation, Log, TEXT("===== AnimNext Modules ====="));
+	UE_LOG(LogAnimation, Log, TEXT("===== UAF Modules ====="));
 	UE_LOG(LogAnimation, Log, TEXT("Num Graphs: %u"), AnimationGraphs.Num());
 
 	for (const UAnimNextAnimationGraph* AnimationGraph : AnimationGraphs)
diff --git a/Engine/Plugins/Experimental/UAF/UAFAnimGraph/Source/UAFAnimGraph/Private/EvaluationVM/Tasks/ApplyAdditiveKeyframe.cpp b/Engine/Plugins/Experimental/UAF/UAFAnimGraph/Source/UAFAnimGraph/Private/EvaluationVM/Tasks/ApplyAdditiveKeyframe.cpp
--- a/Engine/Plugins/Experimental/UAF/UAFAnimGraph/Source/UAFAnimGraph/Private/EvaluationVM/Tasks/ApplyAdditiveKeyframe.cpp
+++ b/Engine/Plugins/Experimental/UAF/UAFAnimGraph/Source/UAFAnimGraph/Private/EvaluationVM/Tasks/ApplyAdditiveKeyframe.cpp
@@ -52,6 +52,7 @@ void FAnimNextApplyAdditiveKeyframeTask::Execute(UE::UAF::FEvaluationVM& VM) con
 		// We have a single input, discard it since it must be the additive pose, either way something went wrong
 		// Push the reference pose since we'll expect a non-additive pose
 		VM.PushValue(KEYFRAME_STACK_NAME, MakeUnique<FKeyframeState>(VM.MakeReferenceKeyframe(false)));
+		UE_LOG(LogAnimation, Warning, TEXT("FAnimNextApplyAdditiveKeyframeTask::Execute: Could not apply additive keyframe - only a single input was provided. Pushing a ref pose instead."));
 		return;
 	}
 
@@ -59,6 +60,7 @@ void FAnimNextApplyAdditiveKeyframeTask::Execute(UE::UAF::FEvaluationVM& VM) con
 	{
 		// Additive must be addtive type, push reference pose if not the case
 		VM.PushValue(KEYFRAME_STACK_NAME, MakeUnique<FKeyframeState>(VM.MakeReferenceKeyframe(false)));
+		UE_LOG(LogAnimation, Warning, TEXT("FAnimNextApplyAdditiveKeyframeTask::Execute: Could not apply additive keyframe - The expected additive pose was not of an additive type. Pushing a ref pose instead."));
 		return;
 	}
 
diff --git a/Engine/Plugins/Experimental/UAF/UAFAnimGraph/Source/UAFAnimGraph/Private/EvaluationVM/Tasks/ApplyAdditiveKeyframePerBone.cpp b/Engine/Plugins/Experimental/UAF/UAFAnimGraph/Source/UAFAnimGraph/Private/EvaluationVM/Tasks/ApplyAdditiveKeyframePerBone.cpp
--- a/Engine/Plugins/Experimental/UAF/UAFAnimGraph/Source/UAFAnimGraph/Private/EvaluationVM/Tasks/ApplyAdditiveKeyframePerBone.cpp
+++ b/Engine/Plugins/Experimental/UAF/UAFAnimGraph/Source/UAFAnimGraph/Private/EvaluationVM/Tasks/ApplyAdditiveKeyframePerBone.cpp
@@ -62,6 +62,7 @@ void FUAFApplyAdditiveKeyframePerBoneTask::Execute(UE::UAF::FEvaluationVM& VM) c
 		// We have a single input, discard it since it must be the additive pose, either way something went wrong
 		// Push the reference pose since we'll expect a non-additive pose
 		VM.PushValue(KEYFRAME_STACK_NAME, MakeUnique<FKeyframeState>(VM.MakeReferenceKeyframe(false)));
+		UE_LOG(LogAnimation, Warning, TEXT("FUAFApplyAdditiveKeyframePerBoneTask::Execute: Could not apply additive keyframe - only a single input was provided. Pushing a ref pose instead."));
 		return;
 	}
 
@@ -69,6 +70,7 @@ void FUAFApplyAdditiveKeyframePerBoneTask::Execute(UE::UAF::FEvaluationVM& VM) c
 	{
 		// Additive must be additive type, push reference pose if not the case
 		VM.PushValue(KEYFRAME_STACK_NAME, MakeUnique<FKeyframeState>(VM.MakeReferenceKeyframe(false)));
+		UE_LOG(LogAnimation, Warning, TEXT("FUAFApplyAdditiveKeyframePerBoneTask::Execute: Could not apply additive keyframe - The expected additive pose was not of an additive type. Pushing a ref pose instead."));
 		return;
 	}
 
diff --git a/Engine/Plugins/Experimental/UAF/UAFAnimGraph/Source/UAFAnimGraph/Private/EvaluationVM/Tasks/ConvertRotationsToLocalSpace.cpp b/Engine/Plugins/Experimental/UAF/UAFAnimGraph/Source/UAFAnimGraph/Private/EvaluationVM/Tasks/ConvertRotationsToLocalSpace.cpp
--- a/Engine/Plugins/Experimental/UAF/UAFAnimGraph/Source/UAFAnimGraph/Private/EvaluationVM/Tasks/ConvertRotationsToLocalSpace.cpp
+++ b/Engine/Plugins/Experimental/UAF/UAFAnimGraph/Source/UAFAnimGraph/Private/EvaluationVM/Tasks/ConvertRotationsToLocalSpace.cpp
@@ -40,6 +40,7 @@ void FAnimNextConvertRotationsToLocalSpaceTask::Execute(UE::UAF::FEvaluationVM&
 		if (!Keyframe || !Keyframe->IsValid())
 		{
 			// No pose to convert
+			UE_LOG(LogAnimation, Warning, TEXT("FAnimNextConvertRotationsToLocalSpaceTask::Execute: Failed to fully execute - No pose data to convert."));
 			return;
 		}
 
diff --git a/Engine/Plugins/Experimental/UAF/UAFAnimGraph/Source/UAFAnimGraph/Private/EvaluationVM/Tasks/ConvertRotationsToMeshSpace.cpp b/Engine/Plugins/Experimental/UAF/UAFAnimGraph/Source/UAFAnimGraph/Private/EvaluationVM/Tasks/ConvertRotationsToMeshSpace.cpp
--- a/Engine/Plugins/Experimental/UAF/UAFAnimGraph/Source/UAFAnimGraph/Private/EvaluationVM/Tasks/ConvertRotationsToMeshSpace.cpp
+++ b/Engine/Plugins/Experimental/UAF/UAFAnimGraph/Source/UAFAnimGraph/Private/EvaluationVM/Tasks/ConvertRotationsToMeshSpace.cpp
@@ -40,6 +40,7 @@ void FAnimNextConvertRotationsToMeshSpaceTask::Execute(UE::UAF::FEvaluationVM& V
 		if (!Keyframe || !Keyframe->IsValid())
 		{
 			// No pose to convert
+			UE_LOG(LogAnimation, Warning, TEXT("FAnimNextConvertRotationsToMeshSpaceTask::Execute: Failed to fully execute - No pose data to convert."));
 			return;
 		}
 
diff --git a/Engine/Plugins/Experimental/UAF/UAFAnimGraph/Source/UAFAnimGraph/Private/EvaluationVM/Tasks/StoreKeyframe.cpp b/Engine/Plugins/Experimental/UAF/UAFAnimGraph/Source/UAFAnimGraph/Private/EvaluationVM/Tasks/StoreKeyframe.cpp
--- a/Engine/Plugins/Experimental/UAF/UAFAnimGraph/Source/UAFAnimGraph/Private/EvaluationVM/Tasks/StoreKeyframe.cpp
+++ b/Engine/Plugins/Experimental/UAF/UAFAnimGraph/Source/UAFAnimGraph/Private/EvaluationVM/Tasks/StoreKeyframe.cpp
@@ -85,6 +85,7 @@ void FAnimNextDuplicateTopKeyframeTask::Execute(UE::UAF::FEvaluationVM& VM) cons
 	{
 		// If there is no Keyframe, generate a reference key frame
 		KeyframeOut = MakeUnique<FKeyframeState>(VM.MakeReferenceKeyframe(false));
+		UE_LOG(LogAnimation, Log, TEXT("FAnimNextDuplicateTopKeyframeTask::Execute: There was no valid keyframe to duplicate. Pushing ref pose keyframe instead."));
 	}
 
 	// Push the copy or reference keyframe to the stack
diff --git a/Engine/Plugins/Experimental/UAF/UAFAnimGraph/Source/UAFAnimGraph/Private/Graph/AnimNextAnimationGraph.cpp b/Engine/Plugins/Experimental/UAF/UAFAnimGraph/Source/UAFAnimGraph/Private/Graph/AnimNextAnimationGraph.cpp
--- a/Engine/Plugins/Experimental/UAF/UAFAnimGraph/Source/UAFAnimGraph/Private/Graph/AnimNextAnimationGraph.cpp
+++ b/Engine/Plugins/Experimental/UAF/UAFAnimGraph/Source/UAFAnimGraph/Private/Graph/AnimNextAnimationGraph.cpp
@@ -54,6 +54,7 @@ TSharedPtr<FAnimNextGraphInstance> UAnimNextAnimationGraph::AllocateInstance(con
 	const FAnimNextTraitHandle ResolvedRootTraitHandle = ResolvedRootTraitHandles.FindRef(EntryPoint);
 	if (!ResolvedRootTraitHandle.IsValid())
 	{
+		UE_LOG(LogAnimation, Error, TEXT("UAnimNextAnimationGraph::AllocateInstance: Could not allocate graph instance - Root trait handle has not been resolved sucessfully."));
 		return nullptr;
 	}
 	
@@ -106,6 +107,7 @@ TSharedPtr<FAnimNextGraphInstance> UAnimNextAnimationGraph::AllocateInstance(con
 	{
 		// We failed to allocate our instance, reset the ptr
 		InstanceImpl.Reset();
+		UE_LOG(LogAnimation, Error, TEXT("UAnimNextAnimationGraph::AllocateInstance: Failed to allocate graph instance."));
 	}
 
 #if WITH_EDITORONLY_DATA
diff --git a/Engine/Plugins/Experimental/UAF/UAFAnimGraph/Source/UAFAnimGraph/Private/Graph/RigUnit_AnimNextRunAnimationGraph_v2.cpp b/Engine/Plugins/Experimental/UAF/UAFAnimGraph/Source/UAFAnimGraph/Private/Graph/RigUnit_AnimNextRunAnimationGraph_v2.cpp
--- a/Engine/Plugins/Experimental/UAF/UAFAnimGraph/Source/UAFAnimGraph/Private/Graph/RigUnit_AnimNextRunAnimationGraph_v2.cpp
+++ b/Engine/Plugins/Experimental/UAF/UAFAnimGraph/Source/UAFAnimGraph/Private/Graph/RigUnit_AnimNextRunAnimationGraph_v2.cpp
@@ -17,6 +17,7 @@
 #include "TraceAnimNextGraphInstances.h"
 #include "Factory/AnimGraphFactory.h"
 #include "UAF/UAFSkeletonUserData.h"
+#include "UAFLogging.h"
 
 #include UE_INLINE_GENERATED_CPP_BY_NAME(RigUnit_AnimNextRunAnimationGraph_v2)
 
@@ -31,6 +32,8 @@ FRigUnit_AnimNextRunAnimationGraph_v2_Execute()
 
 	if (!ReferencePose.ReferencePose.IsValid())
 	{
+		UAF_RIGUNIT_LOG(Warning, TEXT("Could not run graph - Ref Pose is invalid. [Asset: %s] [HostGraph: %s]"), 
+			*GetNameSafe(Graph.Asset), *GetNameSafe(Graph.HostGraph));
 		return;
 	}
 
@@ -56,6 +59,8 @@ FRigUnit_AnimNextRunAnimationGraph_v2_Execute()
 	if(HostGraph == nullptr)
 	{
 		AnimationGraphComponent.ReleaseInstance(WorkData.WeakHost);
+		UAF_RIGUNIT_LOG(Warning, TEXT("Could not run graph - No valid host graph could be found. [Asset: %s] [HostGraph: %s]"),
+			*GetNameSafe(Graph.Asset), *GetNameSafe(Graph.HostGraph));
 		return;
 	}
 
@@ -84,6 +89,8 @@ FRigUnit_AnimNextRunAnimationGraph_v2_Execute()
 
 	if(!WorkData.WeakHost.IsValid())
 	{
+		UAF_RIGUNIT_LOG(Warning, TEXT("Could not run graph - Host Graph could not be allocated. [Asset: %s] [HostGraph: %s]"),
+			*GetNameSafe(Graph.Asset), *GetNameSafe(Graph.HostGraph));
 		return;
 	}
 
diff --git a/Engine/Plugins/Experimental/UAF/UAFAnimGraph/Source/UAFAnimGraph/Private/Graph/RigUnit_GetPostProcessAnimation.cpp b/Engine/Plugins/Experimental/UAF/UAFAnimGraph/Source/UAFAnimGraph/Private/Graph/RigUnit_GetPostProcessAnimation.cpp
--- a/Engine/Plugins/Experimental/UAF/UAFAnimGraph/Source/UAFAnimGraph/Private/Graph/RigUnit_GetPostProcessAnimation.cpp
+++ b/Engine/Plugins/Experimental/UAF/UAFAnimGraph/Source/UAFAnimGraph/Private/Graph/RigUnit_GetPostProcessAnimation.cpp
@@ -7,6 +7,7 @@
 #include "Graph/PostProcessAnimationAssetUserData.h"
 #include "Module/AnimNextModuleInstance.h"
 #include "Module/AnimNextSkeletalMeshComponentReferenceComponent.h"
+#include "UAFLogging.h"
 
 #include UE_INLINE_GENERATED_CPP_BY_NAME(RigUnit_GetPostProcessAnimation)
 
@@ -40,4 +41,8 @@ FRigUnit_GetPostProcessAnimation_Execute()
 			bShouldEvaluate = SkeletalMesh->ShouldEvaluatePostProcessAnimGraph(LODLevel);
 		}
 	}
+	else
+	{
+		UAF_RIGUNIT_LOG(Warning, TEXT("Could not get post process animation - Skeletal mesh component is not valid."));
+	}
 }
diff --git a/Engine/Plugins/Experimental/UAF/UAFAnimGraph/Source/UAFAnimGraph/Private/Injection/InjectionCallbackProxy.cpp b/Engine/Plugins/Experimental/UAF/UAFAnimGraph/Source/UAFAnimGraph/Private/Injection/InjectionCallbackProxy.cpp
--- a/Engine/Plugins/Experimental/UAF/UAFAnimGraph/Source/UAFAnimGraph/Private/Injection/InjectionCallbackProxy.cpp
+++ b/Engine/Plugins/Experimental/UAF/UAFAnimGraph/Source/UAFAnimGraph/Private/Injection/InjectionCallbackProxy.cpp
@@ -1,4 +1,4 @@
-﻿// Copyright Epic Games, Inc. All Rights Reserved.
+// Copyright Epic Games, Inc. All Rights Reserved.
 
 #include "Injection/InjectionCallbackProxy.h"
 #include "Animation/AnimSequence.h"
@@ -33,6 +33,7 @@ bool UInjectionCallbackProxy::Inject(
 {
 	if (AnimNextComponent == nullptr)
 	{
+		UE_LOG(LogAnimation, Warning, TEXT("UInjectionCallbackProxy::Inject: Failed to inject - UAFComponent is not valid."));
 		return false;
 	}
 
diff --git a/Engine/Plugins/Experimental/UAF/UAFAnimGraph/Source/UAFAnimGraph/Private/Injection/InjectionRequest.cpp b/Engine/Plugins/Experimental/UAF/UAFAnimGraph/Source/UAFAnimGraph/Private/Injection/InjectionRequest.cpp
--- a/Engine/Plugins/Experimental/UAF/UAFAnimGraph/Source/UAFAnimGraph/Private/Injection/InjectionRequest.cpp
+++ b/Engine/Plugins/Experimental/UAF/UAFAnimGraph/Source/UAFAnimGraph/Private/Injection/InjectionRequest.cpp
@@ -19,22 +19,26 @@ namespace UE::UAF
 		check(IsInGameThread());
 		if (InHost == nullptr || !InReference.IsValid())
 		{
+			UE_LOG(LogAnimation, Warning, TEXT("FInjectionRequest::Inject: Failed to inject - Provided host or system reference are not valid."));
 			return false;	// Nothing to play
 		}
 
 		if (Status != EInjectionStatus::None)
 		{
+			UE_LOG(LogAnimation, Log, TEXT("FInjectionRequest::Inject: Failed to inject - Request is already playing, can not play again."));
 			return false;	// Already playing, cannot play again
 		}
 
 		if(!ValidateArgs(InRequestArgs))
 		{
+			UE_LOG(LogAnimation, Warning, TEXT("FInjectionRequest::Inject: Failed to inject - Args are not valid."));
 			return false;	// Invalid args
 		}
 
 		UWorld* World = InHost->GetWorld();
 		if(World == nullptr)
 		{
+			UE_LOG(LogAnimation, Warning, TEXT("FInjectionRequest::Inject: Failed to inject - No valid World to play in."));
 			return false;	// No world to play in
 		}
 
@@ -59,10 +63,12 @@ namespace UE::UAF
 		check(IsInGameThread());
 		if (!EnumHasAnyFlags(Status, EInjectionStatus::Playing))
 		{
+			UE_LOG(LogAnimation, Log, TEXT("FInjectionRequest::Uninject: Could not uninject - The request is not playing."));
 			return;	// Not playing
 		}
 		else if (EnumHasAnyFlags(Status, EInjectionStatus::Interrupted))
 		{
+			UE_LOG(LogAnimation, Log, TEXT("FInjectionRequest::Uninject: Could not uninject - The request has already been interrupted."));
 			return;	// We already got interrupted
 		}
 
diff --git a/Engine/Plugins/Experimental/UAF/UAFAnimGraph/Source/UAFAnimGraph/Private/Injection/MontageTrait.cpp b/Engine/Plugins/Experimental/UAF/UAFAnimGraph/Source/UAFAnimGraph/Private/Injection/MontageTrait.cpp
--- a/Engine/Plugins/Experimental/UAF/UAFAnimGraph/Source/UAFAnimGraph/Private/Injection/MontageTrait.cpp
+++ b/Engine/Plugins/Experimental/UAF/UAFAnimGraph/Source/UAFAnimGraph/Private/Injection/MontageTrait.cpp
@@ -20,6 +20,7 @@
 #include "EvaluationVM/Tasks/ApplyAdditiveKeyframe.h"
 #include "EvaluationVM/Tasks/BlendKeyframesPerBone.h"
 #include "EvaluationVM/Tasks/ApplyAdditiveKeyframePerBone.h"
+#include "UAFLogging.h"
 
 namespace UE::UAF
 {
@@ -86,6 +87,7 @@ namespace UE::UAF
 		const FUAFMontageComponent* MontageComponent = InstanceData->MontageComponent;
 		if (MontageComponent == nullptr)
 		{
+			UAF_TRAIT_LOG(Warning, TEXT("FMontageTrait::PreUpdate: Could not update montages - There is no valid MontageComponent."));
 			return;
 		}
 
@@ -98,6 +100,7 @@ namespace UE::UAF
 		const FName CurrentSlotName = SharedData->GetSlotName(Binding);
 		if (CurrentSlotName.IsNone())
 		{
+			UAF_TRAIT_LOG(Warning, TEXT("FMontageTrait::PreUpdate: Could not update montages - Trait has not valid slot name."));
 			return;
 		}
 
diff --git a/Engine/Plugins/Experimental/UAF/UAFAnimGraph/Source/UAFAnimGraph/Private/Injection/UAFMontageComponent.cpp b/Engine/Plugins/Experimental/UAF/UAFAnimGraph/Source/UAFAnimGraph/Private/Injection/UAFMontageComponent.cpp
--- a/Engine/Plugins/Experimental/UAF/UAFAnimGraph/Source/UAFAnimGraph/Private/Injection/UAFMontageComponent.cpp
+++ b/Engine/Plugins/Experimental/UAF/UAFAnimGraph/Source/UAFAnimGraph/Private/Injection/UAFMontageComponent.cpp
@@ -128,12 +128,14 @@ void FUAFMontageComponent::OnBindToInstance()
 	FAnimNextModuleInstance* ModuleInstance = GetModuleInstancePtr();
 	if (!ModuleInstance)
 	{
+		UE_LOG(LogAnimation, Warning, TEXT("FUAFMontageComponent::OnBindToInstance: Could not bind to instance - Module Instance is not valid."));
 		return;
 	}
 
 	UAnimNextComponent* UAFComponent = Cast<UAnimNextComponent>(ModuleInstance->GetObject());
 	if (!UAFComponent)
 	{
+		UE_LOG(LogAnimation, Warning, TEXT("FUAFMontageComponent::OnBindToInstance: Could not bind to instance - UAF Component is not valid."));
 		return;
 	}
 
@@ -163,18 +165,21 @@ void FUAFMontageComponent::OnBindToInstance()
 			UAnimNextComponent* UAFComponent = WeakUAFComponent.Get();
 			if (!UAFComponent)
 			{
+				UE_LOG(LogAnimation, Warning, TEXT("FUAFMontageComponent: Could not setup montage tick function - UAF Component is not valid."));
 				return;
 			}
 
 			AActor* UAFActor = UAFComponent->GetOwner();
 			if (!UAFActor)
 			{
+				UE_LOG(LogAnimation, Warning, TEXT("FUAFMontageComponent: Could not setup montage tick function - Owner of UAF Component [%s] is not valid."), *UAFComponent->GetName());
 				return;
 			}
 
 			USkeletalMeshComponent* SkelMeshComponent = UAFActor->FindComponentByClass<USkeletalMeshComponent>();
 			if (!SkelMeshComponent || SkelMeshComponent->bEnableAnimation)
 			{
+				UE_LOG(LogAnimation, Warning, TEXT("FUAFMontageComponent: Could not setup montage tick function - [%s] does not have a valid SkeletalMeshComponent or Animations are not disabled."), *UAFActor->GetName());
 				return;
 			}
 
@@ -208,6 +213,10 @@ void FUAFMontageComponent::OnBindToInstance()
 
 				ModuleEventTickFunction->AddPrerequisite(UAFComponent, MontageComponent->MontageTickFunction);
 			}
+			else
+			{
+				UE_LOG(LogAnimation, Warning, TEXT("FUAFMontageComponent: Could not setup montage tick function - [%s] has no valid World."), *UAFComponent->GetName());
+			}
 		});
 }
 
diff --git a/Engine/Plugins/Experimental/UAF/UAFAnimGraph/Source/UAFAnimGraph/Private/TraitCore/ExecutionContext.cpp b/Engine/Plugins/Experimental/UAF/UAFAnimGraph/Source/UAFAnimGraph/Private/TraitCore/ExecutionContext.cpp
--- a/Engine/Plugins/Experimental/UAF/UAFAnimGraph/Source/UAFAnimGraph/Private/TraitCore/ExecutionContext.cpp
+++ b/Engine/Plugins/Experimental/UAF/UAFAnimGraph/Source/UAFAnimGraph/Private/TraitCore/ExecutionContext.cpp
@@ -14,6 +14,7 @@
 #include "Graph/AnimNextGraphInstance.h"
 #include "Components/SceneComponent.h"
 #include "Module/AnimNextModuleInstance.h"
+#include "Logging/StructuredLog.h"
 
 namespace UE::UAF
 {
@@ -126,15 +127,16 @@ namespace UE::UAF
 
 		if (!ChildTraitHandle.IsValid())
 		{
+			UE_LOG(LogAnimation, Log, TEXT("FExecutionContext::AllocateNodeInstance: Could not allocate node instance - The provided trait handle is not valid."));
 			return FTraitPtr();	// Attempting to allocate a node using an invalid trait handle
 		}
 
-		if (!ensure(IsBound()))
+		if (!ensureMsgf(IsBound(), TEXT("FExecutionContext::AllocateNodeInstance: Could not allocate node instance - The execution context is not bound to a valid graph instance.")))
 		{
 			return FTraitPtr();	// The execution context must be bound to a valid graph instance
 		}
 
-		if (!ensure(GraphInstance.GetAnimationGraph() != nullptr))
+		if (!ensureMsgf(GraphInstance.GetAnimationGraph() != nullptr, TEXT("FExecutionContext::AllocateNodeInstance: Could not allocate node instance - The provided graph instance has no valid animation graph.")))
 		{
 			return FTraitPtr();	// We need a valid graph instance to allocate into
 		}
@@ -145,12 +147,13 @@ namespace UE::UAF
 		const FNodeDescription& NodeDesc = GetNodeDescription(GraphInstance, ChildNodeHandle);
 		const FNodeTemplate* NodeTemplate = GetNodeTemplate(NodeDesc);
 
-		if (!ensure(NodeTemplate != nullptr))
+		if (!ensureMsgf(NodeTemplate != nullptr, TEXT("FExecutionContext::AllocateNodeInstance: Could not allocate node instance - Node template wasn't found, node descriptor might be corrupted")))
 		{
 			return FTraitPtr();	// Node template wasn't found, node descriptor is perhaps corrupted
 		}
 		else if (ChildTraitIndex >= NodeTemplate->GetNumTraits())
 		{
+			UE_LOG(LogAnimation, Warning, TEXT("FExecutionContext::AllocateNodeInstance: Could not allocate node instance - The requested trait index doesn't exist on that node descriptor."));
 			return FTraitPtr();	// The requested trait index doesn't exist on that node descriptor
 		}
 
@@ -211,15 +214,16 @@ namespace UE::UAF
 		
 		if (!ChildTraitHandle.IsValid())
 		{
+			UE_LOG(LogAnimation, Log, TEXT("FExecutionContext::AllocateNodeInstance: Could not allocate node instance - The provided trait handle is not valid."));
 			return FTraitPtr();	// Attempting to allocate a node using an invalid trait handle
 		}
 
-		if (!ensure(IsBound()))
+		if (!ensureMsgf(IsBound(), TEXT("FExecutionContext::AllocateNodeInstance: Could not allocate node instance - The execution context is not bound to a valid graph instance.")))
 		{
 			return FTraitPtr();	// The execution context must be bound to a valid graph instance
 		}
 
-		if (!ensure(ParentBinding.IsValid()))
+		if (!ensureMsgf(ParentBinding.IsValid(), TEXT("FExecutionContext::AllocateNodeInstance: Could not allocate node instance - The provided parent binding was not valid.")))
 		{
 			return FTraitPtr();	// We need a parent binding to know which graph instance to allocate into
 		}
@@ -233,12 +237,13 @@ namespace UE::UAF
 		const FNodeDescription& NodeDesc = GetNodeDescription(GraphInstance, ChildNodeHandle);
 		const FNodeTemplate* NodeTemplate = GetNodeTemplate(NodeDesc);
 
-		if (!ensure(NodeTemplate != nullptr))
+		if (!ensureMsgf(NodeTemplate != nullptr, TEXT("FExecutionContext::AllocateNodeInstance: Could not allocate node instance - Node template wasn't found, node descriptor might be corrupted")))
 		{
 			return FTraitPtr();	// Node template wasn't found, node descriptor is perhaps corrupted
 		}
 		else if (ChildTraitIndex >= NodeTemplate->GetNumTraits())
 		{
+			UE_LOG(LogAnimation, Warning, TEXT("FExecutionContext::AllocateNodeInstance: Could not allocate node instance - The requested trait index doesn't exist on that node descriptor."));
 			return FTraitPtr();	// The requested trait index doesn't exist on that node descriptor
 		}
 
@@ -313,7 +318,7 @@ namespace UE::UAF
 
 		FNodeInstance* NodeInstance = NodePtr.GetNodeInstance();
 
-		if (!ensure(IsBoundTo(NodeInstance->GetOwner())))
+		if (!ensureMsgf(IsBoundTo(NodeInstance->GetOwner()), TEXT("FExecutionContext::ReleaseNodeInstance: The execution context isn't bound to the right graph instance")))
 		{
 			return;	// The execution context isn't bound to the right graph instance
 		}
@@ -324,12 +329,13 @@ namespace UE::UAF
 
 		if (NodeInstance->RemoveReference())
 		{
+			UE_LOG(LogAnimation, Warning, TEXT("FExecutionContext::ReleaseNodeInstance: Could not release node instance - Node Instance still has references"));
 			return;	// Node instance still has references, we can't release it
 		}
 
 		const FNodeDescription& NodeDesc = GetNodeDescription(*NodeInstance);
 		const FNodeTemplate* NodeTemplate = GetNodeTemplate(NodeDesc);
-		if (!ensure(NodeTemplate != nullptr))
+		if (!ensureMsgf(NodeTemplate != nullptr, TEXT("FExecutionContext::ReleaseNodeInstance: Could not release node instance - Node template wasn't found, node descriptor is perhaps corrupted (we'll leak the node memory)")))
 		{
 			return;	// Node template wasn't found, node descriptor is perhaps corrupted (we'll leak the node memory)
 		}
@@ -384,10 +390,11 @@ namespace UE::UAF
 		if (!TraitPtr.IsValid())
 		{
 			OutStackBinding.Reset();
+			UE_LOG(LogAnimation, Warning, TEXT("FExecutionContext::GetStack: Provided trait pointer was not valid."));
 			return false;
 		}
 
-		if (!ensure(IsBoundTo(TraitPtr.GetNodeInstance()->GetOwner())))
+		if (!ensureMsgf(IsBoundTo(TraitPtr.GetNodeInstance()->GetOwner()), TEXT("FExecutionContext::GetStack: The execution context isn't bound to the right graph instance")))
 		{
 			OutStackBinding.Reset();
 			return false;	// The execution context isn't bound to the right graph instance
@@ -401,6 +408,7 @@ namespace UE::UAF
 	{
 		if (!Binding.IsValid())
 		{
+			UE_LOG(LogAnimation, Warning, TEXT("FExecutionContext::PushScopedInterfaceImpl: Cannot push scoped interface - Provided trait binding is not valid."));
 			return;	// Don't queue invalid pointers
 		}
 
@@ -440,6 +448,7 @@ namespace UE::UAF
 	{
 		if (!Binding.IsValid())
 		{
+			UE_LOG(LogAnimation, Warning, TEXT("FExecutionContext::PopScopedInterfaceImpl: Cannot pop scoped interface - Provided trait binding is not valid."));
 			return false;
 		}
 
@@ -469,6 +478,7 @@ namespace UE::UAF
 	{
 		if (!StackBinding.IsValid())
 		{
+			UE_LOG(LogAnimation, Warning, TEXT("FExecutionContext::PopStackScopedInterfaces: Cannot pop scoped interfaces - Provided stack binding is not valid."));
 			return false;
 		}
 
@@ -624,4 +634,34 @@ namespace UE::UAF
 		return FTransform::Identity;
 	}
 #endif
+
+	void FExecutionContext::Log(ELogVerbosity::Type Verbosity, const FString& InMessage, const FString& TraitName)
+	{
+		const UObject* HostObject = GetHostObject();
+		const UObject* SystemOwner = HostObject ? HostObject->GetOuter() : nullptr;
+		const UObject* RootGraphAsset = IsBound() ? RootGraphInstance->GetAsset<UObject>() : nullptr;
+
+		switch (Verbosity)
+		{
+		case ELogVerbosity::Fatal:
+			UE_LOGFMT(LogAnimation, Fatal, "{TraitName}: {Message} - [Root Graph: {GraphPath}] [Owner: {OwnerPath}]", TraitName, InMessage, FSoftObjectPath(RootGraphAsset), FSoftObjectPath(SystemOwner));
+			break;
+		case ELogVerbosity::Error:
+			UE_LOGFMT(LogAnimation, Error, "{TraitName}: {Message} - [Root Graph: {GraphPath}] [Owner: {OwnerPath}]", TraitName, InMessage, FSoftObjectPath(RootGraphAsset), FSoftObjectPath(SystemOwner));
+			break;
+		case ELogVerbosity::Warning:
+			UE_LOGFMT(LogAnimation, Warning, "{TraitName}: {Message} - [Root Graph: {GraphPath}] [Owner: {OwnerPath}]", TraitName, InMessage, FSoftObjectPath(RootGraphAsset), FSoftObjectPath(SystemOwner));
+			break;
+		case ELogVerbosity::Display:
+			UE_LOGFMT(LogAnimation, Display, "{TraitName}: {Message} - [Root Graph: {GraphPath}] [Owner: {OwnerPath}]", TraitName, InMessage, FSoftObjectPath(RootGraphAsset), FSoftObjectPath(SystemOwner));
+			break;
+		case ELogVerbosity::Log:
+			UE_LOGFMT(LogAnimation, Log, "{TraitName}: {Message} - [Root Graph: {GraphPath}] [Owner: {OwnerPath}]", TraitName, InMessage, FSoftObjectPath(RootGraphAsset), FSoftObjectPath(SystemOwner));
+			break;
+		default:
+			UE_LOGFMT(LogAnimation, Error, "Unsupported log category used! Message: {TraitName}: {Message} - [Root Graph: {GraphPath}] [Owner: {OwnerPath}]", TraitName, InMessage, FSoftObjectPath(RootGraphAsset), FSoftObjectPath(SystemOwner));
+			break;
+		}
+	}
+
 }
diff --git a/Engine/Plugins/Experimental/UAF/UAFAnimGraph/Source/UAFAnimGraph/Private/TraitCore/TraitEventRaising.cpp b/Engine/Plugins/Experimental/UAF/UAFAnimGraph/Source/UAFAnimGraph/Private/TraitCore/TraitEventRaising.cpp
--- a/Engine/Plugins/Experimental/UAF/UAFAnimGraph/Source/UAFAnimGraph/Private/TraitCore/TraitEventRaising.cpp
+++ b/Engine/Plugins/Experimental/UAF/UAFAnimGraph/Source/UAFAnimGraph/Private/TraitCore/TraitEventRaising.cpp
@@ -1,4 +1,4 @@
-﻿// Copyright Epic Games, Inc. All Rights Reserved.
+// Copyright Epic Games, Inc. All Rights Reserved.
 
 #include "TraitCore/TraitEventRaising.h"
 #include "TraitCore/TraitEvent.h"
@@ -6,6 +6,7 @@
 #include "TraitCore/TraitBinding.h"
 #include "TraitCore/TraitStackBinding.h"
 #include "Graph/UAFGraphInstanceComponent.h"
+#include "EngineLogs.h"
 
 namespace UE::UAF
 {
@@ -13,6 +14,7 @@ namespace UE::UAF
 	{
 		if (!StackBinding.IsValid() || !Event.IsValid())
 		{
+			UE_LOG(LogAnimation, Warning, TEXT("RaiseTraitEvent: Could not raise trait event - Either the stack binding or trait event are invalid."));
 			return;
 		}
 
@@ -33,6 +35,7 @@ namespace UE::UAF
 	{
 		if (!StackBinding.IsValid())
 		{
+			UE_LOG(LogAnimation, Warning, TEXT("RaiseTraitEvents: Could not raise trait events - The stack binding is invalid."));
 			return;
 		}
 
diff --git a/Engine/Plugins/Experimental/UAF/UAFAnimGraph/Source/UAFAnimGraph/Private/TraitCore/TraitInterfaceRegistry.cpp b/Engine/Plugins/Experimental/UAF/UAFAnimGraph/Source/UAFAnimGraph/Private/TraitCore/TraitInterfaceRegistry.cpp
--- a/Engine/Plugins/Experimental/UAF/UAFAnimGraph/Source/UAFAnimGraph/Private/TraitCore/TraitInterfaceRegistry.cpp
+++ b/Engine/Plugins/Experimental/UAF/UAFAnimGraph/Source/UAFAnimGraph/Private/TraitCore/TraitInterfaceRegistry.cpp
@@ -3,6 +3,7 @@
 #include "TraitCore/TraitInterfaceRegistry.h"
 
 #include "TraitCore/ITraitInterface.h"
+#include "EngineLogs.h"
 
 namespace UE::UAF
 {
@@ -91,6 +92,7 @@ const ITraitInterface* FTraitInterfaceRegistry::Find(FTraitInterfaceUID TraitInt
 {
 	if (!TraitInterfaceUID.IsValid())
 	{
+		UE_LOG(LogAnimation, Warning, TEXT("FTraitInterfaceRegistry::Find: Provided trait interface UID is invalid."));
 		return nullptr;
 	}
 
@@ -107,6 +109,7 @@ void FTraitInterfaceRegistry::Register(const TSharedPtr<ITraitInterface>& TraitI
 {
 	if (TraitInterface == nullptr)
 	{
+		UE_LOG(LogAnimation, Warning, TEXT("FTraitInterfaceRegistry::Register: Could not register trait interface - Provided trait interface is invalid."));
 		return;
 	}
 
@@ -122,6 +125,7 @@ void FTraitInterfaceRegistry::Unregister(const TSharedPtr<ITraitInterface>& Trai
 {
 	if (TraitInterface == nullptr)
 	{
+		UE_LOG(LogAnimation, Warning, TEXT("FTraitInterfaceRegistry::Unregister: Could not unregister trait interface - Provided trait interface is invalid."));
 		return;
 	}
 
diff --git a/Engine/Plugins/Experimental/UAF/UAFAnimGraph/Source/UAFAnimGraph/Private/Traits/SynchronizeUsingGroups.cpp b/Engine/Plugins/Experimental/UAF/UAFAnimGraph/Source/UAFAnimGraph/Private/Traits/SynchronizeUsingGroups.cpp
--- a/Engine/Plugins/Experimental/UAF/UAFAnimGraph/Source/UAFAnimGraph/Private/Traits/SynchronizeUsingGroups.cpp
+++ b/Engine/Plugins/Experimental/UAF/UAFAnimGraph/Source/UAFAnimGraph/Private/Traits/SynchronizeUsingGroups.cpp
@@ -6,6 +6,7 @@
 #include "Factory/AnimNextFactoryParams.h"
 #include "Graph/SyncGroup_GraphInstanceComponent.h"
 #include "GraphInterfaces/AnimNextNativeDataInterface_SynchronizeUsingGroups.h"
+#include "UAFLogging.h"
 
 namespace UE::UAF
 {
@@ -56,6 +57,7 @@ namespace UE::UAF
 		if (!InstanceData->bHasTimelinePlayer)
 		{
 			// If we don't have a timeline player beneath us on the trait stack, there is nothing to synchronize
+			UAF_TRAIT_LOG(Warning, TEXT("FSynchronizeUsingGroupsTrait::PreUpdate: Could not synchronize - No valid timeline player beneath us on the trait stack."));
 			IUpdate::PreUpdate(Context, Binding, TraitState);
 			return;
 		}
diff --git a/Engine/Plugins/Experimental/UAF/UAFAnimGraph/Source/UAFAnimGraph/Public/TraitCore/ExecutionContext.h b/Engine/Plugins/Experimental/UAF/UAFAnimGraph/Source/UAFAnimGraph/Public/TraitCore/ExecutionContext.h
--- a/Engine/Plugins/Experimental/UAF/UAFAnimGraph/Source/UAFAnimGraph/Public/TraitCore/ExecutionContext.h
+++ b/Engine/Plugins/Experimental/UAF/UAFAnimGraph/Source/UAFAnimGraph/Public/TraitCore/ExecutionContext.h
@@ -213,6 +213,9 @@ namespace UE::UAF
 		// Get the debug draw interface
 		UE_API FRigVMDrawInterface* GetDebugDrawInterface() const;
 #endif
+		
+		void Log(ELogVerbosity::Type Verbosity, const FString& Message, const FString& TraitName);
+
 	private:
 		// No copy or move
 		FExecutionContext(const FExecutionContext&) = delete;
diff --git a/Engine/Plugins/Experimental/UAF/UAFAnimGraph/Source/UAFAnimGraph/Public/Traits/BlendSmoother.h b/Engine/Plugins/Experimental/UAF/UAFAnimGraph/Source/UAFAnimGraph/Public/Traits/BlendSmoother.h
--- a/Engine/Plugins/Experimental/UAF/UAFAnimGraph/Source/UAFAnimGraph/Public/Traits/BlendSmoother.h
+++ b/Engine/Plugins/Experimental/UAF/UAFAnimGraph/Source/UAFAnimGraph/Public/Traits/BlendSmoother.h
@@ -125,7 +125,7 @@ namespace UE::UAF
 		virtual UCurveFloat* GetCustomBlendCurve(FExecutionContext& Context, const TTraitBinding<ISmoothBlend>& Binding, int32 ChildIndex) const override;
 
 #if WITH_EDITOR
-		virtual bool IsHidden() const override { return false; } // overrdie the base, as it is hidden
+		virtual bool IsHidden() const override { return false; } // override the base, as it is hidden
 #endif
 	};
 }
diff --git a/Engine/Plugins/Experimental/UAF/UAFPoseSearch/Source/UAFPoseSearch/Private/RigUnit_GenerateCharacterMovementComponentTrajectory.cpp b/Engine/Plugins/Experimental/UAF/UAFPoseSearch/Source/UAFPoseSearch/Private/RigUnit_GenerateCharacterMovementComponentTrajectory.cpp
--- a/Engine/Plugins/Experimental/UAF/UAFPoseSearch/Source/UAFPoseSearch/Private/RigUnit_GenerateCharacterMovementComponentTrajectory.cpp
+++ b/Engine/Plugins/Experimental/UAF/UAFPoseSearch/Source/UAFPoseSearch/Private/RigUnit_GenerateCharacterMovementComponentTrajectory.cpp
@@ -5,6 +5,7 @@
 #include "GameFramework/CharacterMovementComponent.h"
 #include "PoseSearch/PoseSearchDefines.h"
 #include "PoseSearch/PoseSearchTrajectoryLibrary.h"
+#include "UAFLogging.h"
 
 FRigUnit_GenerateCharacterMovementComponentTrajectory_Execute()
 {
@@ -13,13 +14,14 @@ FRigUnit_GenerateCharacterMovementComponentTrajectory_Execute()
 
 	if (CharacterMovementComponent == nullptr)
 	{
+		UAF_RIGUNIT_LOG(Warning, TEXT("Failed to get Character Movement Component."));
 		return;
 	}
 
 	AActor* Owner = CharacterMovementComponent->GetOwner();
 	if (Owner == nullptr)
 	{
-		UE_LOG(LogPoseSearch, Warning, TEXT("FRigUnit_GenerateCharacterMovementComponentTrajectory_Execute: Failed to get Character Movement Component's actor owner."));
+		UAF_RIGUNIT_LOG(Warning, TEXT("Failed to get Character Movement Component's actor owner."));
 		return;
 	}
 	
diff --git a/Engine/Plugins/Runtime/RigVM/Source/RigVM/Public/RigVMCore/RigVMExecuteContext.h b/Engine/Plugins/Runtime/RigVM/Source/RigVM/Public/RigVMCore/RigVMExecuteContext.h
--- a/Engine/Plugins/Runtime/RigVM/Source/RigVM/Public/RigVMCore/RigVMExecuteContext.h
+++ b/Engine/Plugins/Runtime/RigVM/Source/RigVM/Public/RigVMCore/RigVMExecuteContext.h
@@ -306,7 +306,7 @@ struct FRigVMExecuteContext : public FRigVMExecutePin
 
 	virtual ~FRigVMExecuteContext() = default;
 
-	void Log(const FRigVMLogSettings& InLogSettings, const FString& InMessage) const
+	virtual void Log(const FRigVMLogSettings& InLogSettings, const FString& InMessage) const
 	{
 		if(RuntimeSettings.LogFunction.IsValid())
 		{
