Prevent terminating blends in case of UAF PreUpdate pass running with 0.0 delta time value

#Synced-CL 48786226

diff --git a/Engine/Plugins/Experimental/UAF/UAFAnimGraph/Source/UAFAnimGraph/Private/Traits/BlendSmoother.cpp b/Engine/Plugins/Experimental/UAF/UAFAnimGraph/Source/UAFAnimGraph/Private/Traits/BlendSmoother.cpp
--- a/Engine/Plugins/Experimental/UAF/UAFAnimGraph/Source/UAFAnimGraph/Private/Traits/BlendSmoother.cpp
+++ b/Engine/Plugins/Experimental/UAF/UAFAnimGraph/Source/UAFAnimGraph/Private/Traits/BlendSmoother.cpp
@@ -145,6 +145,12 @@ namespace UE::UAF
 				ChildBlendData.Weight *= ReciprocalSum;
 			}
 		}
+		
+		// Do not attempt to terminate blends in case of being ticked with zero delta-time
+		if (DeltaTime == 0.f)
+		{
+			return;
+		}
 
 		TTraitBinding<IDiscreteBlend> DiscreteBlendTrait;
 		Binding.GetStackInterface(DiscreteBlendTrait);
