Update modify curve task to support the new UAF attribute runtime

#Synced-CL 48779969

diff --git a/Engine/Plugins/Experimental/UAF/UAFAnimGraph/Source/UAFAnimGraph/Private/Traits/ModifyCurveTrait.cpp b/Engine/Plugins/Experimental/UAF/UAFAnimGraph/Source/UAFAnimGraph/Private/Traits/ModifyCurveTrait.cpp
--- a/Engine/Plugins/Experimental/UAF/UAFAnimGraph/Source/UAFAnimGraph/Private/Traits/ModifyCurveTrait.cpp
+++ b/Engine/Plugins/Experimental/UAF/UAFAnimGraph/Source/UAFAnimGraph/Private/Traits/ModifyCurveTrait.cpp
@@ -53,8 +53,40 @@ void FModifyCurveTask::Execute(UE::UAF::FEvaluationVM& VM) const
 
 	if (VM.GetContext().IsValid())
 	{
-		// TODO: Implement with new attribute runtime
-		return;
+		if (const TUniquePtr<FAttributeCollection>* Collection = VM.PeekValue<TUniquePtr<FAttributeCollection>>(ATTRIBUTE_STACK_NAME, 0))
+		{
+			FAttributeNamedSetPtr NamedSet = (*Collection)->GetNamedSet();
+			if (FAttributeTypedSetPtr FloatTypedSet = NamedSet->FindTypedSet<FFloatAnimationAttribute>())
+			{
+				FAttributeMappingKey MappingKey = FAttributeMappingKey::MakeFromTo<FFloatAnimationAttribute>();
+				
+				if (TAttributeMap<FFloatAnimationAttribute>* FloatMap = (*Collection)->FindMap<FFloatAnimationAttribute>())
+				{
+					for (const FModifyCurveParameters& CurveParameters : SharedData->ModifyCurveParameters)
+					{
+						if (FFloatAnimationAttribute* Float = FloatMap->Find(CurveParameters.CurveName))
+						{
+							const float CurrentValue = (*Float).Value;
+							const float UpdatedValue = ProcessCurveOperation(CurrentValue, CurveParameters.CurveValue, InstanceData->Alpha, SharedData->ApplyMode);
+							(*Float).Value = UpdatedValue;
+						}
+					}
+				}
+
+				if (TAttributeSetMap<FFloatAnimationAttribute>*	FloatSetMap = (*Collection)->FindSetMap<FFloatAnimationAttribute>(MappingKey))
+				{
+					for (const FModifyCurveParameters& CurveParameters : SharedData->ModifyCurveParameters)
+					{
+						if (const FAttributeSetIndex RootMotionAttributeIndex = FloatTypedSet->FindIndex(CurveParameters.CurveName))
+						{
+							const float CurrentValue = (*FloatSetMap)[RootMotionAttributeIndex].Value;
+							const float UpdatedValue = ProcessCurveOperation(CurrentValue, CurveParameters.CurveValue, InstanceData->Alpha, SharedData->ApplyMode);
+							(*FloatSetMap)[RootMotionAttributeIndex].Value = UpdatedValue;
+						}
+					}
+				}
+			}
+		}
 	}
 	
 	if (const TUniquePtr<FKeyframeState>* Keyframe = VM.PeekValue<TUniquePtr<FKeyframeState>>(KEYFRAME_STACK_NAME, 0))
@@ -65,7 +97,6 @@ void FModifyCurveTask::Execute(UE::UAF::FEvaluationVM& VM) const
 			float UpdatedValue = ProcessCurveOperation(CurrentValue, CurveParameters.CurveValue, InstanceData->Alpha, SharedData->ApplyMode);
 			Keyframe->Get()->Curves.Set(CurveParameters.CurveName, UpdatedValue);
 		}
-		
 	}
 }
 
