Add the ability to push/pop the currently active binding/named set in UAF

#Synced-CL 48913404

diff --git a/Engine/Plugins/Experimental/MovieSceneAnimMixer/Source/MovieSceneAnimMixer/Private/Systems/MovieSceneAnimMixerSystem.cpp b/Engine/Plugins/Experimental/MovieSceneAnimMixer/Source/MovieSceneAnimMixer/Private/Systems/MovieSceneAnimMixerSystem.cpp
--- a/Engine/Plugins/Experimental/MovieSceneAnimMixer/Source/MovieSceneAnimMixer/Private/Systems/MovieSceneAnimMixerSystem.cpp
+++ b/Engine/Plugins/Experimental/MovieSceneAnimMixer/Source/MovieSceneAnimMixer/Private/Systems/MovieSceneAnimMixerSystem.cpp
@@ -350,7 +350,7 @@ void FAnimNextBlendTwoKeyframesPreserveRootMotionTask::Execute(UE::UAF::FEvaluat
 	using namespace UE::UAF;
 	using namespace UE::MovieScene;
 
-	if (VM.GetContext().IsValid())
+	if (VM.GetActiveNamedSet())
 	{
 		// TODO: Implement with new attribute runtime
 		TUniquePtr<FAttributeCollection> DiscardedCollection;
@@ -485,7 +485,7 @@ void FMovieSceneAccumulateAbsoluteBlendTask::Execute(UE::UAF::FEvaluationVM& VM)
 {
 	using namespace UE::UAF;
 
-	if (VM.GetContext().IsValid())
+	if (VM.GetActiveNamedSet())
 	{
 		// TODO: Implement with new attribute runtime
 		TUniquePtr<FAttributeCollection> DiscardedCollection;
diff --git a/Engine/Plugins/Experimental/MovieSceneAnimMixer/Source/MovieSceneAnimMixer/Private/Systems/MovieSceneMixedSkeletalAnimationSystem.cpp b/Engine/Plugins/Experimental/MovieSceneAnimMixer/Source/MovieSceneAnimMixer/Private/Systems/MovieSceneMixedSkeletalAnimationSystem.cpp
--- a/Engine/Plugins/Experimental/MovieSceneAnimMixer/Source/MovieSceneAnimMixer/Private/Systems/MovieSceneMixedSkeletalAnimationSystem.cpp
+++ b/Engine/Plugins/Experimental/MovieSceneAnimMixer/Source/MovieSceneAnimMixer/Private/Systems/MovieSceneMixedSkeletalAnimationSystem.cpp
@@ -405,10 +405,10 @@ void FMovieSceneSkeletalAnimationEvaluationTask::Execute(UE::UAF::FEvaluationVM&
 	using namespace UE::UAF;
 	using namespace UE::MovieScene;
 
-	if (VM.GetContext().IsValid())
+	if (VM.GetActiveNamedSet())
 	{
 		// TODO: Implement with new attribute runtime
-		FAttributeCollectionStack Collection(VM.GetContext().GetNamedSet());
+		FAttributeCollectionStack Collection(VM.GetActiveNamedSet());
 		Collection.InitWithValueSpace(FValueSpace(EValueSpaceType::Local));
 
 		VM.PushValue(ATTRIBUTE_STACK_NAME, MakeUnique<FAttributeCollection>(MoveTemp(Collection)));
diff --git a/Engine/Plugins/Experimental/MovieSceneAnimMixer/Source/MovieSceneAnimMixer/Private/Systems/MovieSceneRootMotionSystem.cpp b/Engine/Plugins/Experimental/MovieSceneAnimMixer/Source/MovieSceneAnimMixer/Private/Systems/MovieSceneRootMotionSystem.cpp
--- a/Engine/Plugins/Experimental/MovieSceneAnimMixer/Source/MovieSceneAnimMixer/Private/Systems/MovieSceneRootMotionSystem.cpp
+++ b/Engine/Plugins/Experimental/MovieSceneAnimMixer/Source/MovieSceneAnimMixer/Private/Systems/MovieSceneRootMotionSystem.cpp
@@ -270,7 +270,7 @@ void FAnimNextConvertRootMotionToWorldSpaceTask::Execute(UE::UAF::FEvaluationVM&
 	using namespace UE::UAF; 
 	using namespace UE::MovieScene;
 
-	if (VM.GetContext().IsValid())
+	if (VM.GetActiveNamedSet())
 	{
 		// TODO: Implement with new attribute runtime
 		return;
@@ -363,7 +363,7 @@ void FAnimNextStoreRootTransformTask::Execute(UE::UAF::FEvaluationVM& VM) const
 	using namespace UE::UAF;
 	using namespace UE::MovieScene;
 
-	if (VM.GetContext().IsValid())
+	if (VM.GetActiveNamedSet())
 	{
 		// TODO: Implement with new attribute runtime
 		return;
diff --git a/Engine/Plugins/Experimental/MovieScenePoseSearchTracks/Source/MovieScenePoseSearchTracks/Private/Sections/MovieSceneStitchAnimSection.cpp b/Engine/Plugins/Experimental/MovieScenePoseSearchTracks/Source/MovieScenePoseSearchTracks/Private/Sections/MovieSceneStitchAnimSection.cpp
--- a/Engine/Plugins/Experimental/MovieScenePoseSearchTracks/Source/MovieScenePoseSearchTracks/Private/Sections/MovieSceneStitchAnimSection.cpp
+++ b/Engine/Plugins/Experimental/MovieScenePoseSearchTracks/Source/MovieScenePoseSearchTracks/Private/Sections/MovieSceneStitchAnimSection.cpp
@@ -59,10 +59,10 @@ void FMovieSceneStitchAnimEvaluationTask::Execute(UE::UAF::FEvaluationVM& VM) co
 	using namespace UE::PoseSearch;
 	using namespace UE::Anim;
 
-	if (VM.GetContext().IsValid())
+	if (VM.GetActiveNamedSet())
 	{
 		// TODO: Implement with new attribute runtime
-		FAttributeCollectionStack Collection(VM.GetContext().GetNamedSet());
+		FAttributeCollectionStack Collection(VM.GetActiveNamedSet());
 		Collection.InitWithValueSpace(FValueSpace(EValueSpaceType::Local));
 
 		VM.PushValue(ATTRIBUTE_STACK_NAME, MakeUnique<FAttributeCollection>(MoveTemp(Collection)));
diff --git a/Engine/Plugins/Experimental/RigLogicUAF/Source/RigLogicUAF/Private/RigLogicTask.cpp b/Engine/Plugins/Experimental/RigLogicUAF/Source/RigLogicUAF/Private/RigLogicTask.cpp
--- a/Engine/Plugins/Experimental/RigLogicUAF/Source/RigLogicUAF/Private/RigLogicTask.cpp
+++ b/Engine/Plugins/Experimental/RigLogicUAF/Source/RigLogicUAF/Private/RigLogicTask.cpp
@@ -28,7 +28,7 @@ void FUAFRigLogicTask::Execute(UE::UAF::FEvaluationVM& VM) const
 {
 	SCOPE_CYCLE_COUNTER(STAT_UAF_Task_RigLogic);
 
-	if (VM.GetContext().IsValid())
+	if (VM.GetActiveNamedSet())
 	{
 		// TODO: Implement with new attribute runtime
 		return;
diff --git a/Engine/Plugins/Experimental/UAF/UAF/Source/UAF/Private/UAF/Attributes/AttributeEvaluationContext.cpp b/Engine/Plugins/Experimental/UAF/UAF/Source/UAF/Private/UAF/Attributes/AttributeEvaluationContext.cpp
--- a/Engine/Plugins/Experimental/UAF/UAF/Source/UAF/Private/UAF/Attributes/AttributeEvaluationContext.cpp
+++ b/Engine/Plugins/Experimental/UAF/UAF/Source/UAF/Private/UAF/Attributes/AttributeEvaluationContext.cpp
@@ -8,19 +8,14 @@ namespace UE::UAF
 {
 	FAttributeEvaluationContext::FAttributeEvaluationContext() = default;
 
-	FAttributeEvaluationContext::FAttributeEvaluationContext(const UAbstractSkeletonSetBinding* InSetBinding, FName InSetName)
+	FAttributeEvaluationContext::FAttributeEvaluationContext(const UAbstractSkeletonSetBinding& InSetBinding, FName InSetName)
 	{
-		if (InSetBinding == nullptr)
-		{
-			return;
-		}
-
-		if (const FAttributeBindingDataPtr& InBindingData = InSetBinding->GetAttributeBindingData())
+		if (const FAttributeBindingDataPtr& InBindingData = InSetBinding.GetAttributeBindingData())
 		{
 			if (FAttributeNamedSetPtr InNamedSet = InBindingData->FindNamedSet(InSetName))
 			{
-				SetBinding = InSetBinding;
-				Skeleton = InSetBinding->GetSkeleton();
+				SetBinding = &InSetBinding;
+				Skeleton = InSetBinding.GetSkeleton();
 				BindingData = InBindingData;
 				NamedSet = MoveTemp(InNamedSet);
 				DefaultAttributeValues = NamedSet->GetDefaultAttributeValues();
@@ -43,6 +38,22 @@ namespace UE::UAF
 		return NamedSet;
 	}
 
+	bool FAttributeEvaluationContext::SetNamedSet(FName InSetName)
+	{
+		if (BindingData)
+		{
+			if (FAttributeNamedSetPtr InNamedSet = BindingData->FindNamedSet(InSetName))
+			{
+				NamedSet = MoveTemp(InNamedSet);
+				DefaultAttributeValues = NamedSet->GetDefaultAttributeValues();
+				return true;
+			}
+		}
+
+		// Named set not found
+		return false;
+	}
+
 	const FAttributeCollectionPtr& FAttributeEvaluationContext::GetDefaultAttributeValues() const
 	{
 		return DefaultAttributeValues;
diff --git a/Engine/Plugins/Experimental/UAF/UAF/Source/UAF/Public/UAF/Attributes/AttributeEvaluationContext.h b/Engine/Plugins/Experimental/UAF/UAF/Source/UAF/Public/UAF/Attributes/AttributeEvaluationContext.h
--- a/Engine/Plugins/Experimental/UAF/UAF/Source/UAF/Public/UAF/Attributes/AttributeEvaluationContext.h
+++ b/Engine/Plugins/Experimental/UAF/UAF/Source/UAF/Public/UAF/Attributes/AttributeEvaluationContext.h
@@ -22,7 +22,7 @@ namespace UE::UAF
 	{
 	public:
 		UE_API FAttributeEvaluationContext();
-		UE_API FAttributeEvaluationContext(const UAbstractSkeletonSetBinding* SetBinding, FName SetName);
+		UE_API FAttributeEvaluationContext(const UAbstractSkeletonSetBinding& SetBinding, FName SetName);
 
 		// Returns whether or not this context is valid
 		[[nodiscard]] UE_API bool IsValid() const;
@@ -36,6 +36,10 @@ namespace UE::UAF
 		// Returns the named set we are currently evaluating with (or nullptr if we aren't valid)
 		[[nodiscard]] UE_API const FAttributeNamedSetPtr& GetNamedSet() const;
 
+		// Sets the named set
+		// Returns true if we succeeded and the set exists in the current binding, false otherwise
+		[[nodiscard]] UE_API bool SetNamedSet(FName SetName);
+
 		// Returns the default attribute values for the current named set (or nullptr if we aren't valid)
 		[[nodiscard]] UE_API const FAttributeCollectionPtr& GetDefaultAttributeValues() const;
 
diff --git a/Engine/Plugins/Experimental/UAF/UAFAnimGraph/Source/UAFAnimGraph/Private/EvaluationVM/EvaluationVM.cpp b/Engine/Plugins/Experimental/UAF/UAFAnimGraph/Source/UAFAnimGraph/Private/EvaluationVM/EvaluationVM.cpp
--- a/Engine/Plugins/Experimental/UAF/UAFAnimGraph/Source/UAFAnimGraph/Private/EvaluationVM/EvaluationVM.cpp
+++ b/Engine/Plugins/Experimental/UAF/UAFAnimGraph/Source/UAFAnimGraph/Private/EvaluationVM/EvaluationVM.cpp
@@ -34,56 +34,45 @@ namespace UE::UAF
 		, CurrentLOD(InCurrentLOD)
 		, EvaluationFlags(InEvaluationFlags)
 	{
-		//const UE::Anim::FCurveFilterSettings CurveFilterSettings(UE::Anim::ECurveFilterMode::DisallowAll);
-		//if(USkeleton* Skeleton = const_cast<USkeleton*>(InReferencePose.Skeleton.Get())) // const_cast because the bone container takes a mutable reference
-		{
-			//BoneContainer.InitializeTo(InReferencePose.GetLODBoneIndexToMeshBoneIndexMap(InCurrentLOD), CurveFilterSettings, *Skeleton);
-
-			// TODO: In AnimInstanceProxy this is how we initialize the bone container, we need to get the component somehow or we
-			// gotta figure out how to support ref pose overrides
-#if 0
-			// Use the shared bone container
-			RequiredBones = Component->GetSharedRequiredBones();
-
-			// The first anim instance will initialize the required bones, all others will re-use it
-			if (!RequiredBones->IsValid())
-			{
-				RequiredBones->InitializeTo(Component->RequiredBones, Component->GetCurveFilterSettings(), *Asset);
-
-				// If there is a ref pose override, we want to replace ref pose in RequiredBones
-				// Update ref pose in required bones structure (either set it, or clear it, depending on if one is set on the Component)
-				RequiredBones->SetRefPoseOverride(Component->GetRefPoseOverride());
-			}
-#endif
-		}
+		// Push an invalid context
+		EvaluationContextStack.Push(FAttributeEvaluationContext());
 	}
 
 	FEvaluationVM::FEvaluationVM(EEvaluationFlags InEvaluationFlags, const UAbstractSkeletonSetBinding& InSetBinding, FName InSetName, int32 InCurrentLOD)
-		: InternalContext(&InSetBinding, InSetName)
-		, AttributeValueTransformerMap(FAttributeRuntimeRegistry::Get().GetTransformerMap())
+		: AttributeValueTransformerMap(FAttributeRuntimeRegistry::Get().GetTransformerMap())
 		, CurrentLOD(InCurrentLOD)
 		, EvaluationFlags(InEvaluationFlags)
 	{
+		EvaluationContextStack.Push(FAttributeEvaluationContext(InSetBinding, InSetName));
 	}
 
 	// REMOVE ME: Exclusively for testing that before/after yields the same results
 	FEvaluationVM::FEvaluationVM(EEvaluationFlags InEvaluationFlags, const FReferencePose& InReferencePose, const UAbstractSkeletonSetBinding* InSetBinding, FName InSetName, int32 InCurrentLOD)
 		: ReferencePose(&InReferencePose)
-		, InternalContext(InSetBinding, InSetName)
 		, AttributeValueTransformerMap(FAttributeRuntimeRegistry::Get().GetTransformerMap())
 		, CurrentLOD(InCurrentLOD)
 		, EvaluationFlags(InEvaluationFlags)
 	{
+		if (InSetBinding != nullptr)
+		{
+			EvaluationContextStack.Push(FAttributeEvaluationContext(*InSetBinding, InSetName));
+		}
+		else
+		{
+			// Push an invalid context
+			EvaluationContextStack.Push(FAttributeEvaluationContext());
+		}
 	}
 
 	bool FEvaluationVM::IsValid() const
 	{
-		return ReferencePose != nullptr || InternalContext.IsValid();
+		return ReferencePose != nullptr || EvaluationContextStack.Last().IsValid();
 	}
 
 	void FEvaluationVM::Shrink()
 	{
 		InternalStacks.Shrink();
+		EvaluationContextStack.Shrink();
 	}
 
 	FKeyframeState FEvaluationVM::MakeReferenceKeyframe(bool bAdditiveKeyframe) const
@@ -130,9 +119,53 @@ namespace UE::UAF
 		return Keyframe;
 	}
 
-	const FAttributeEvaluationContext& FEvaluationVM::GetContext() const
+	const FAttributeEvaluationContext& FEvaluationVM::GetActiveEvaluationContext() const
+	{
+		return EvaluationContextStack.Last();
+	}
+
+	const FAttributeNamedSetPtr& FEvaluationVM::GetActiveNamedSet() const
 	{
-		return InternalContext;
+		return EvaluationContextStack.Last().GetNamedSet();
+	}
+
+	bool FEvaluationVM::PushEvaluationContext(const UAbstractSkeletonSetBinding& SetBinding, FName SetName)
+	{
+		FAttributeEvaluationContext NewContext(SetBinding, SetName);
+		if (NewContext.IsValid())
+		{
+			EvaluationContextStack.Push(MoveTemp(NewContext));
+			return true;
+		}
+
+		// New evaluation context isn't valid
+		return false;
+	}
+
+	bool FEvaluationVM::PushEvaluationContext(FName SetName)
+	{
+		// Copy the currently active context
+		FAttributeEvaluationContext NewContext = EvaluationContextStack.Last();
+		if (NewContext.SetNamedSet(SetName))
+		{
+			EvaluationContextStack.Push(MoveTemp(NewContext));
+			return true;
+		}
+
+		// Failed to find the named set in our currently active binding
+		return false;
+	}
+
+	bool FEvaluationVM::PopEvaluationContext()
+	{
+		if (EvaluationContextStack.Num() > 1)
+		{
+			EvaluationContextStack.Pop(EAllowShrinking::No);
+			return true;
+		}
+
+		// Cannot pop an entry from an empty stack or the last entry
+		return false;
 	}
 
 	const FAttributeValueTransformerMapPtr& FEvaluationVM::GetTransformerMap() const
diff --git a/Engine/Plugins/Experimental/UAF/UAFAnimGraph/Source/UAFAnimGraph/Private/EvaluationVM/Tasks/ApplyAdditiveKeyframe.cpp b/Engine/Plugins/Experimental/UAF/UAFAnimGraph/Source/UAFAnimGraph/Private/EvaluationVM/Tasks/ApplyAdditiveKeyframe.cpp
--- a/Engine/Plugins/Experimental/UAF/UAFAnimGraph/Source/UAFAnimGraph/Private/EvaluationVM/Tasks/ApplyAdditiveKeyframe.cpp
+++ b/Engine/Plugins/Experimental/UAF/UAFAnimGraph/Source/UAFAnimGraph/Private/EvaluationVM/Tasks/ApplyAdditiveKeyframe.cpp
@@ -32,13 +32,13 @@ void FAnimNextApplyAdditiveKeyframeTask::Execute(UE::UAF::FEvaluationVM& VM) con
 {
 	using namespace UE::UAF;
 
-	if (VM.GetContext().IsValid())
+	if (VM.GetActiveNamedSet())
 	{
 		TUniquePtr<FAttributeCollection> AdditiveCollection;
 		if (!VM.PopValue(ATTRIBUTE_STACK_NAME, AdditiveCollection))
 		{
 			// We have no inputs, nothing to do
-			FAttributeCollectionStack Collection(VM.GetContext().GetNamedSet());
+			FAttributeCollectionStack Collection(VM.GetActiveNamedSet());
 			Collection.InitWithValueSpace(FValueSpace(EValueSpaceType::Local));
 
 			VM.PushValue(ATTRIBUTE_STACK_NAME, MakeUnique<FAttributeCollection>(MoveTemp(Collection)));
@@ -51,7 +51,7 @@ void FAnimNextApplyAdditiveKeyframeTask::Execute(UE::UAF::FEvaluationVM& VM) con
 		{
 			// We have a single input, discard it since it must be the additive pose, either way something went wrong
 			// Push the reference pose since we'll expect a non-additive pose
-			FAttributeCollectionStack Collection(VM.GetContext().GetNamedSet());
+			FAttributeCollectionStack Collection(VM.GetActiveNamedSet());
 			Collection.InitWithValueSpace(FValueSpace(EValueSpaceType::Local));
 
 			VM.PushValue(ATTRIBUTE_STACK_NAME, MakeUnique<FAttributeCollection>(MoveTemp(Collection)));
@@ -63,7 +63,7 @@ void FAnimNextApplyAdditiveKeyframeTask::Execute(UE::UAF::FEvaluationVM& VM) con
 		if (!AdditiveSpace.IsAdditive())
 		{
 			// Additive input must be additive, push reference pose if not the case
-			FAttributeCollectionStack Collection(VM.GetContext().GetNamedSet());
+			FAttributeCollectionStack Collection(VM.GetActiveNamedSet());
 			Collection.InitWithValueSpace(FValueSpace(EValueSpaceType::Local));
 
 			VM.PushValue(ATTRIBUTE_STACK_NAME, MakeUnique<FAttributeCollection>(MoveTemp(Collection)));
@@ -81,6 +81,7 @@ void FAnimNextApplyAdditiveKeyframeTask::Execute(UE::UAF::FEvaluationVM& VM) con
 			bConvertBaseToMeshSpace = BaseSpace.GetMixedSpaceFlags() != EMixedSpaceFlags::MeshRotation;
 		}
 
+		// Modify our base input in-place
 		if (bConvertBaseToMeshSpace)
 		{
 			Transformers::FBoneSpace::LocalToMeshRotation(*BaseCollection, *BaseCollection);
@@ -120,7 +121,7 @@ void FAnimNextApplyAdditiveKeyframeTask::Execute(UE::UAF::FEvaluationVM& VM) con
 
 		if (!AdditiveKeyframe->Pose.IsAdditive())
 		{
-			// Additive must be addtive type, push reference pose if not the case
+			// Additive must be additive type, push reference pose if not the case
 			VM.PushValue(KEYFRAME_STACK_NAME, MakeUnique<FKeyframeState>(VM.MakeReferenceKeyframe(false)));
 			UE_LOG(LogAnimation, Warning, TEXT("FAnimNextApplyAdditiveKeyframeTask::Execute: Could not apply additive keyframe - The expected additive pose was not of an additive type. Pushing a ref pose instead."));
 			return;
diff --git a/Engine/Plugins/Experimental/UAF/UAFAnimGraph/Source/UAFAnimGraph/Private/EvaluationVM/Tasks/ApplyAdditiveKeyframePerBone.cpp b/Engine/Plugins/Experimental/UAF/UAFAnimGraph/Source/UAFAnimGraph/Private/EvaluationVM/Tasks/ApplyAdditiveKeyframePerBone.cpp
--- a/Engine/Plugins/Experimental/UAF/UAFAnimGraph/Source/UAFAnimGraph/Private/EvaluationVM/Tasks/ApplyAdditiveKeyframePerBone.cpp
+++ b/Engine/Plugins/Experimental/UAF/UAFAnimGraph/Source/UAFAnimGraph/Private/EvaluationVM/Tasks/ApplyAdditiveKeyframePerBone.cpp
@@ -40,7 +40,7 @@ void FUAFApplyAdditiveKeyframePerBoneTask::Execute(UE::UAF::FEvaluationVM& VM) c
 		return;
 	}
 
-	if (VM.GetContext().IsValid())
+	if (VM.GetActiveNamedSet())
 	{
 		// TODO: Implement with new attribute runtime
 		TUniquePtr<FAttributeCollection> DiscardedCollection;
diff --git a/Engine/Plugins/Experimental/UAF/UAFAnimGraph/Source/UAFAnimGraph/Private/EvaluationVM/Tasks/BlendKeyframes.cpp b/Engine/Plugins/Experimental/UAF/UAFAnimGraph/Source/UAFAnimGraph/Private/EvaluationVM/Tasks/BlendKeyframes.cpp
--- a/Engine/Plugins/Experimental/UAF/UAFAnimGraph/Source/UAFAnimGraph/Private/EvaluationVM/Tasks/BlendKeyframes.cpp
+++ b/Engine/Plugins/Experimental/UAF/UAFAnimGraph/Source/UAFAnimGraph/Private/EvaluationVM/Tasks/BlendKeyframes.cpp
@@ -35,7 +35,7 @@ void FAnimNextBlendTwoKeyframesTask::Execute(UE::UAF::FEvaluationVM& VM) const
 {
 	using namespace UE::UAF;
 
-	if (VM.GetContext().IsValid())
+	if (VM.GetActiveNamedSet())
 	{
 		// Pop our top two poses, we'll re-use the top keyframe for our result
 
@@ -216,7 +216,7 @@ void FAnimNextBlendOverwriteKeyframeWithScaleTask::Execute(UE::UAF::FEvaluationV
 {
 	using namespace UE::UAF;
 
-	if (VM.GetContext().IsValid())
+	if (VM.GetActiveNamedSet())
 	{
 		TUniquePtr<FAttributeCollection> Collection;
 		if (!VM.PopValue(ATTRIBUTE_STACK_NAME, Collection))
@@ -277,7 +277,7 @@ void FAnimNextBlendAddKeyframeWithScaleTask::Execute(UE::UAF::FEvaluationVM& VM)
 {
 	using namespace UE::UAF;
 
-	if (VM.GetContext().IsValid())
+	if (VM.GetActiveNamedSet())
 	{
 		// Pop our top two poses, we'll re-use the top keyframe for our result
 
diff --git a/Engine/Plugins/Experimental/UAF/UAFAnimGraph/Source/UAFAnimGraph/Private/EvaluationVM/Tasks/BlendKeyframesPerBone.cpp b/Engine/Plugins/Experimental/UAF/UAFAnimGraph/Source/UAFAnimGraph/Private/EvaluationVM/Tasks/BlendKeyframesPerBone.cpp
--- a/Engine/Plugins/Experimental/UAF/UAFAnimGraph/Source/UAFAnimGraph/Private/EvaluationVM/Tasks/BlendKeyframesPerBone.cpp
+++ b/Engine/Plugins/Experimental/UAF/UAFAnimGraph/Source/UAFAnimGraph/Private/EvaluationVM/Tasks/BlendKeyframesPerBone.cpp
@@ -39,7 +39,7 @@ void FAnimNextBlendOverwriteKeyframePerBoneWithScaleTask::Execute(UE::UAF::FEval
 {
 	using namespace UE::UAF;
 
-	if (VM.GetContext().IsValid())
+	if (VM.GetActiveNamedSet())
 	{
 		// TODO: Implement with new attribute runtime
 		Super::Execute(VM);
@@ -165,7 +165,7 @@ void FAnimNextBlendAddKeyframePerBoneWithScaleTask::Execute(UE::UAF::FEvaluation
 {
 	using namespace UE::UAF;
 
-	if (VM.GetContext().IsValid())
+	if (VM.GetActiveNamedSet())
 	{
 		// TODO: Implement with new attribute runtime
 		Super::Execute(VM);
@@ -298,7 +298,7 @@ void FAnimNextBlendKeyframePerBoneWithScaleTask::Execute(UE::UAF::FEvaluationVM&
 {
 	using namespace UE::UAF;
 
-	if (VM.GetContext().IsValid())
+	if (VM.GetActiveNamedSet())
 	{
 		// TODO: Implement with new attribute runtime
 		Super::Execute(VM);
diff --git a/Engine/Plugins/Experimental/UAF/UAFAnimGraph/Source/UAFAnimGraph/Private/EvaluationVM/Tasks/ConvertRotationsToLocalSpace.cpp b/Engine/Plugins/Experimental/UAF/UAFAnimGraph/Source/UAFAnimGraph/Private/EvaluationVM/Tasks/ConvertRotationsToLocalSpace.cpp
--- a/Engine/Plugins/Experimental/UAF/UAFAnimGraph/Source/UAFAnimGraph/Private/EvaluationVM/Tasks/ConvertRotationsToLocalSpace.cpp
+++ b/Engine/Plugins/Experimental/UAF/UAFAnimGraph/Source/UAFAnimGraph/Private/EvaluationVM/Tasks/ConvertRotationsToLocalSpace.cpp
@@ -24,7 +24,7 @@ void FAnimNextConvertRotationsToLocalSpaceTask::Execute(UE::UAF::FEvaluationVM&
 {
 	using namespace UE::UAF;
 
-	if (VM.GetContext().IsValid())
+	if (VM.GetActiveNamedSet())
 	{
 		for (int32 PoseIndex = 0; PoseIndex < NumPoses; ++PoseIndex)
 		{
diff --git a/Engine/Plugins/Experimental/UAF/UAFAnimGraph/Source/UAFAnimGraph/Private/EvaluationVM/Tasks/ConvertRotationsToMeshSpace.cpp b/Engine/Plugins/Experimental/UAF/UAFAnimGraph/Source/UAFAnimGraph/Private/EvaluationVM/Tasks/ConvertRotationsToMeshSpace.cpp
--- a/Engine/Plugins/Experimental/UAF/UAFAnimGraph/Source/UAFAnimGraph/Private/EvaluationVM/Tasks/ConvertRotationsToMeshSpace.cpp
+++ b/Engine/Plugins/Experimental/UAF/UAFAnimGraph/Source/UAFAnimGraph/Private/EvaluationVM/Tasks/ConvertRotationsToMeshSpace.cpp
@@ -24,7 +24,7 @@ void FAnimNextConvertRotationsToMeshSpaceTask::Execute(UE::UAF::FEvaluationVM& V
 {
 	using namespace UE::UAF;
 
-	if (VM.GetContext().IsValid())
+	if (VM.GetActiveNamedSet())
 	{
 		for (int32 PoseIndex = 0; PoseIndex < NumPoses; ++PoseIndex)
 		{
diff --git a/Engine/Plugins/Experimental/UAF/UAFAnimGraph/Source/UAFAnimGraph/Private/EvaluationVM/Tasks/CopyBones.cpp b/Engine/Plugins/Experimental/UAF/UAFAnimGraph/Source/UAFAnimGraph/Private/EvaluationVM/Tasks/CopyBones.cpp
--- a/Engine/Plugins/Experimental/UAF/UAFAnimGraph/Source/UAFAnimGraph/Private/EvaluationVM/Tasks/CopyBones.cpp
+++ b/Engine/Plugins/Experimental/UAF/UAFAnimGraph/Source/UAFAnimGraph/Private/EvaluationVM/Tasks/CopyBones.cpp
@@ -18,7 +18,7 @@ void FAnimNextCopyBonesComponentSpaceTask::Execute(UE::UAF::FEvaluationVM& VM) c
 {
 	using namespace UE::UAF;
 
-	if (VM.GetContext().IsValid())
+	if (VM.GetActiveNamedSet())
 	{
 		// TODO: Implement with new attribute runtime
 		return;
diff --git a/Engine/Plugins/Experimental/UAF/UAFAnimGraph/Source/UAFAnimGraph/Private/EvaluationVM/Tasks/DeadBlending.cpp b/Engine/Plugins/Experimental/UAF/UAFAnimGraph/Source/UAFAnimGraph/Private/EvaluationVM/Tasks/DeadBlending.cpp
--- a/Engine/Plugins/Experimental/UAF/UAFAnimGraph/Source/UAFAnimGraph/Private/EvaluationVM/Tasks/DeadBlending.cpp
+++ b/Engine/Plugins/Experimental/UAF/UAFAnimGraph/Source/UAFAnimGraph/Private/EvaluationVM/Tasks/DeadBlending.cpp
@@ -484,7 +484,7 @@ void FAnimNextDeadBlendingTransitionTask::Execute(UE::UAF::FEvaluationVM& VM) co
 {
 	using namespace UE::UAF;
 
-	if (VM.GetContext().IsValid())
+	if (VM.GetActiveNamedSet())
 	{
 		// TODO: Implement with new attribute runtime
 		return;
@@ -559,7 +559,7 @@ void FAnimNextDeadBlendingApplyTask::Execute(UE::UAF::FEvaluationVM& VM) const
 {
 	using namespace UE::UAF;
 
-	if (VM.GetContext().IsValid())
+	if (VM.GetActiveNamedSet())
 	{
 		// TODO: Implement with new attribute runtime
 		return;
diff --git a/Engine/Plugins/Experimental/UAF/UAFAnimGraph/Source/UAFAnimGraph/Private/EvaluationVM/Tasks/NormalizeRotations.cpp b/Engine/Plugins/Experimental/UAF/UAFAnimGraph/Source/UAFAnimGraph/Private/EvaluationVM/Tasks/NormalizeRotations.cpp
--- a/Engine/Plugins/Experimental/UAF/UAFAnimGraph/Source/UAFAnimGraph/Private/EvaluationVM/Tasks/NormalizeRotations.cpp
+++ b/Engine/Plugins/Experimental/UAF/UAFAnimGraph/Source/UAFAnimGraph/Private/EvaluationVM/Tasks/NormalizeRotations.cpp
@@ -14,7 +14,7 @@ void FAnimNextNormalizeKeyframeRotationsTask::Execute(UE::UAF::FEvaluationVM& VM
 {
 	using namespace UE::UAF;
 
-	if (VM.GetContext().IsValid())
+	if (VM.GetActiveNamedSet())
 	{
 		TUniquePtr<FAttributeCollection> Collection;
 		if (!VM.PopValue(ATTRIBUTE_STACK_NAME, Collection))
diff --git a/Engine/Plugins/Experimental/UAF/UAFAnimGraph/Source/UAFAnimGraph/Private/EvaluationVM/Tasks/PushAnimSequenceKeyframe.cpp b/Engine/Plugins/Experimental/UAF/UAFAnimGraph/Source/UAFAnimGraph/Private/EvaluationVM/Tasks/PushAnimSequenceKeyframe.cpp
--- a/Engine/Plugins/Experimental/UAF/UAFAnimGraph/Source/UAFAnimGraph/Private/EvaluationVM/Tasks/PushAnimSequenceKeyframe.cpp
+++ b/Engine/Plugins/Experimental/UAF/UAFAnimGraph/Source/UAFAnimGraph/Private/EvaluationVM/Tasks/PushAnimSequenceKeyframe.cpp
@@ -34,7 +34,7 @@ void FAnimNextAnimSequenceKeyframeTask::Execute(UE::UAF::FEvaluationVM& VM) cons
 {
 	using namespace UE::UAF;
 	
-	const FAttributeEvaluationContext& VMContext = VM.GetContext();
+	const FAttributeEvaluationContext& VMContext = VM.GetActiveEvaluationContext();
 
 	if(const UAnimSequence* AnimSequencePtr = AnimSequence.Get())
 	{
@@ -57,7 +57,7 @@ void FAnimNextAnimSequenceKeyframeTask::Execute(UE::UAF::FEvaluationVM& VM) cons
 				ValueSpace = FValueSpace(EValueSpaceType::Local, bIsAdditive);
 			}
 
-			FAttributeCollectionStack Collection(VM.GetContext().GetNamedSet());
+			FAttributeCollectionStack Collection(VMContext.GetNamedSet());
 			Collection.InitWithValueSpace(ValueSpace);
 
 			FDecompressionTools::GetAnimationPose(AnimSequencePtr, ExtractionContext, VMContext, Collection, bUseRawData);
@@ -128,7 +128,7 @@ void FAnimNextAnimSequenceKeyframeTask::Execute(UE::UAF::FEvaluationVM& VM) cons
 	{
 		if (VMContext.IsValid())
 		{
-			FAttributeCollectionStack Collection(VM.GetContext().GetNamedSet());
+			FAttributeCollectionStack Collection(VMContext.GetNamedSet());
 			Collection.InitWithValueSpace(FValueSpace(EValueSpaceType::Local));
 
 			VM.PushValue(ATTRIBUTE_STACK_NAME, MakeUnique<FAttributeCollection>(MoveTemp(Collection)));
diff --git a/Engine/Plugins/Experimental/UAF/UAFAnimGraph/Source/UAFAnimGraph/Private/EvaluationVM/Tasks/PushPose.cpp b/Engine/Plugins/Experimental/UAF/UAFAnimGraph/Source/UAFAnimGraph/Private/EvaluationVM/Tasks/PushPose.cpp
--- a/Engine/Plugins/Experimental/UAF/UAFAnimGraph/Source/UAFAnimGraph/Private/EvaluationVM/Tasks/PushPose.cpp
+++ b/Engine/Plugins/Experimental/UAF/UAFAnimGraph/Source/UAFAnimGraph/Private/EvaluationVM/Tasks/PushPose.cpp
@@ -29,10 +29,10 @@ void FAnimNextPushPoseTask::Execute(UE::UAF::FEvaluationVM& VM) const
 
 	SCOPE_CYCLE_COUNTER(STAT_AnimNext_Task_PushPose);
 
-	if (VM.GetContext().IsValid())
+	if (VM.GetActiveNamedSet())
 	{
 		// TODO: Implement with new attribute runtime
-		FAttributeCollectionStack Collection(VM.GetContext().GetNamedSet());
+		FAttributeCollectionStack Collection(VM.GetActiveNamedSet());
 		Collection.InitWithValueSpace(FValueSpace(EValueSpaceType::Local));
 
 		VM.PushValue(ATTRIBUTE_STACK_NAME, MakeUnique<FAttributeCollection>(MoveTemp(Collection)));
diff --git a/Engine/Plugins/Experimental/UAF/UAFAnimGraph/Source/UAFAnimGraph/Private/EvaluationVM/Tasks/PushReferenceKeyframe.cpp b/Engine/Plugins/Experimental/UAF/UAFAnimGraph/Source/UAFAnimGraph/Private/EvaluationVM/Tasks/PushReferenceKeyframe.cpp
--- a/Engine/Plugins/Experimental/UAF/UAFAnimGraph/Source/UAFAnimGraph/Private/EvaluationVM/Tasks/PushReferenceKeyframe.cpp
+++ b/Engine/Plugins/Experimental/UAF/UAFAnimGraph/Source/UAFAnimGraph/Private/EvaluationVM/Tasks/PushReferenceKeyframe.cpp
@@ -23,9 +23,9 @@ void FAnimNextPushReferenceKeyframeTask::Execute(UE::UAF::FEvaluationVM& VM) con
 {
 	using namespace UE::UAF;
 
-	if (VM.GetContext().IsValid())
+	if (VM.GetActiveNamedSet())
 	{
-		FAttributeCollectionStack Collection(VM.GetContext().GetNamedSet());
+		FAttributeCollectionStack Collection(VM.GetActiveNamedSet());
 		Collection.InitWithValueSpace(FValueSpace(EValueSpaceType::Local, bIsAdditive));
 
 		VM.PushValue(ATTRIBUTE_STACK_NAME, MakeUnique<FAttributeCollection>(MoveTemp(Collection)));
diff --git a/Engine/Plugins/Experimental/UAF/UAFAnimGraph/Source/UAFAnimGraph/Private/EvaluationVM/Tasks/StoreKeyframe.cpp b/Engine/Plugins/Experimental/UAF/UAFAnimGraph/Source/UAFAnimGraph/Private/EvaluationVM/Tasks/StoreKeyframe.cpp
--- a/Engine/Plugins/Experimental/UAF/UAFAnimGraph/Source/UAFAnimGraph/Private/EvaluationVM/Tasks/StoreKeyframe.cpp
+++ b/Engine/Plugins/Experimental/UAF/UAFAnimGraph/Source/UAFAnimGraph/Private/EvaluationVM/Tasks/StoreKeyframe.cpp
@@ -32,7 +32,7 @@ void FAnimNextStoreKeyframeTransformsTask::Execute(UE::UAF::FEvaluationVM& VM) c
 {
 	using namespace UE::UAF;
 
-	if (VM.GetContext().IsValid())
+	if (VM.GetActiveNamedSet())
 	{
 		// TODO: Implement with new attribute runtime
 		return;
@@ -61,10 +61,10 @@ void FAnimNextDuplicateTopKeyframeTask::Execute(UE::UAF::FEvaluationVM& VM) cons
 {
 	using namespace UE::UAF;
 
-	if (VM.GetContext().IsValid())
+	if (VM.GetActiveNamedSet())
 	{
 		// TODO: Implement with new attribute runtime
-		FAttributeCollectionStack Collection(VM.GetContext().GetNamedSet());
+		FAttributeCollectionStack Collection(VM.GetActiveNamedSet());
 		Collection.InitWithValueSpace(FValueSpace(EValueSpaceType::Local));
 
 		VM.PushValue(ATTRIBUTE_STACK_NAME, MakeUnique<FAttributeCollection>(MoveTemp(Collection)));
diff --git a/Engine/Plugins/Experimental/UAF/UAFAnimGraph/Source/UAFAnimGraph/Private/Traits/ModifyCurveTrait.cpp b/Engine/Plugins/Experimental/UAF/UAFAnimGraph/Source/UAFAnimGraph/Private/Traits/ModifyCurveTrait.cpp
--- a/Engine/Plugins/Experimental/UAF/UAFAnimGraph/Source/UAFAnimGraph/Private/Traits/ModifyCurveTrait.cpp
+++ b/Engine/Plugins/Experimental/UAF/UAFAnimGraph/Source/UAFAnimGraph/Private/Traits/ModifyCurveTrait.cpp
@@ -51,7 +51,7 @@ void FModifyCurveTask::Execute(UE::UAF::FEvaluationVM& VM) const
 {
 	using namespace UE::UAF;
 
-	if (VM.GetContext().IsValid())
+	if (VM.GetActiveNamedSet())
 	{
 		if (const TUniquePtr<FAttributeCollection>* Collection = VM.PeekValue<TUniquePtr<FAttributeCollection>>(ATTRIBUTE_STACK_NAME, 0))
 		{
diff --git a/Engine/Plugins/Experimental/UAF/UAFAnimGraph/Source/UAFAnimGraph/Public/EvaluationVM/EvaluationVM.h b/Engine/Plugins/Experimental/UAF/UAFAnimGraph/Source/UAFAnimGraph/Public/EvaluationVM/EvaluationVM.h
--- a/Engine/Plugins/Experimental/UAF/UAFAnimGraph/Source/UAFAnimGraph/Public/EvaluationVM/EvaluationVM.h
+++ b/Engine/Plugins/Experimental/UAF/UAFAnimGraph/Source/UAFAnimGraph/Public/EvaluationVM/EvaluationVM.h
@@ -151,7 +151,6 @@ namespace UE::UAF
 	 */
 	struct FEvaluationVM
 	{
-		FEvaluationVM() = default;
 		UE_API FEvaluationVM(EEvaluationFlags InEvaluationFlags, const FReferencePose& InReferencePose, int32 InCurrentLOD);
 		UE_API FEvaluationVM(EEvaluationFlags EvaluationFlags, const UAbstractSkeletonSetBinding& SetBinding, FName SetName, int32 CurrentLOD);
 
@@ -198,8 +197,23 @@ namespace UE::UAF
 		// Returns an uninitialized keyframe with memory pre-allocated
 		[[nodiscard]] UE_API FKeyframeState MakeUninitializedKeyframe(bool bAdditiveKeyframe) const;
 
-		// Returns the VM context that controls what state we produce
-		[[nodiscard]] UE_API const FAttributeEvaluationContext& GetContext() const;
+		// Returns the evaluation context that controls what state we produce
+		[[nodiscard]] UE_API const FAttributeEvaluationContext& GetActiveEvaluationContext() const;
+
+		// Returns the named set we are currently evaluating with
+		[[nodiscard]] UE_API const FAttributeNamedSetPtr& GetActiveNamedSet() const;
+
+		// Pushes a new binding and named set to evaluate with
+		// Returns true if we succeeded and the set exists in the specified binding, false otherwise
+		[[nodiscard]] UE_API bool PushEvaluationContext(const UAbstractSkeletonSetBinding& SetBinding, FName SetName);
+
+		// Pushes a new named set to evaluate with
+		// Returns true if we succeeded and the set exists in the currently active binding, false otherwise
+		[[nodiscard]] UE_API bool PushEvaluationContext(FName SetName);
+
+		// Pops the current active named set
+		// Returns whether or not we succeeded (we need to retain at least one named set)
+		UE_API bool PopEvaluationContext();
 
 		// Returns the current LOD
 		[[nodiscard]] int32 GetCurrentLOD() const { return CurrentLOD; }
@@ -228,8 +242,7 @@ namespace UE::UAF
 		const FReferencePose* ReferencePose = nullptr;
 
 		// The internal context that controls what state we are producing
-		// TODO: Should be a stack, push/pop a set mask
-		FAttributeEvaluationContext InternalContext;
+		TArray<FAttributeEvaluationContext> EvaluationContextStack;
 
 		// Shared map of attribute value transformers
 		FAttributeValueTransformerMapPtr AttributeValueTransformerMap;
diff --git a/Engine/Plugins/Experimental/UAF/UAFControlRig/Source/UAFControlRig/Private/ControlRigTask.cpp b/Engine/Plugins/Experimental/UAF/UAFControlRig/Source/UAFControlRig/Private/ControlRigTask.cpp
--- a/Engine/Plugins/Experimental/UAF/UAFControlRig/Source/UAFControlRig/Private/ControlRigTask.cpp
+++ b/Engine/Plugins/Experimental/UAF/UAFControlRig/Source/UAFControlRig/Private/ControlRigTask.cpp
@@ -89,7 +89,7 @@ void FAnimNextControlRigTask::Execute(UE::UAF::FEvaluationVM& VM) const
 
 	using namespace UE::UAF;
 
-	if (VM.GetContext().IsValid())
+	if (VM.GetActiveNamedSet())
 	{
 		// TODO: Implement with new attribute runtime
 		return;
diff --git a/Engine/Plugins/Experimental/UAF/UAFMirroring/Source/UAFMirroring/Private/MirroringTrait.cpp b/Engine/Plugins/Experimental/UAF/UAFMirroring/Source/UAFMirroring/Private/MirroringTrait.cpp
--- a/Engine/Plugins/Experimental/UAF/UAFMirroring/Source/UAFMirroring/Private/MirroringTrait.cpp
+++ b/Engine/Plugins/Experimental/UAF/UAFMirroring/Source/UAFMirroring/Private/MirroringTrait.cpp
@@ -227,7 +227,7 @@ void FAnimNextEvaluationMirroringTask::Execute(UE::UAF::FEvaluationVM& VM) const
 {
 	QUICK_SCOPE_CYCLE_COUNTER(FAnimNextEvaluationMirroringTask_Execute);
 
-	if (VM.GetContext().IsValid())
+	if (VM.GetActiveNamedSet())
 	{
 		// TODO: Implement with new attribute runtime
 		return;
diff --git a/Engine/Plugins/Experimental/UAF/UAFPoseSearch/Source/UAFPoseSearch/Private/HistoryCollectorTrait.cpp b/Engine/Plugins/Experimental/UAF/UAFPoseSearch/Source/UAFPoseSearch/Private/HistoryCollectorTrait.cpp
--- a/Engine/Plugins/Experimental/UAF/UAFPoseSearch/Source/UAFPoseSearch/Private/HistoryCollectorTrait.cpp
+++ b/Engine/Plugins/Experimental/UAF/UAFPoseSearch/Source/UAFPoseSearch/Private/HistoryCollectorTrait.cpp
@@ -107,7 +107,7 @@ void FAnimNextHistoryCollectorPreEvaluateTask::Execute(UE::UAF::FEvaluationVM& V
 	using namespace UE::UAF;
 	using namespace UE::PoseSearch;
 
-	if (VM.GetContext().IsValid())
+	if (VM.GetActiveNamedSet())
 	{
 		FPoseHistoryEvaluationHelper PoseHistoryEvalHelper;
 		PoseHistoryEvalHelper.PoseHistoryPtr = InstanceData->PoseHistoryPtr;
@@ -128,7 +128,7 @@ void FAnimNextHistoryCollectorTask::Execute(UE::UAF::FEvaluationVM& VM) const
 {
 	using namespace UE::UAF;
 
-	if (VM.GetContext().IsValid())
+	if (VM.GetActiveNamedSet())
 	{
 		// TODO: Implement with new attribute runtime
 		return;
diff --git a/Engine/Plugins/Experimental/UAF/UAFPoseSearch/Source/UAFPoseSearch/Private/MotionMatchingTrait.cpp b/Engine/Plugins/Experimental/UAF/UAFPoseSearch/Source/UAFPoseSearch/Private/MotionMatchingTrait.cpp
--- a/Engine/Plugins/Experimental/UAF/UAFPoseSearch/Source/UAFPoseSearch/Private/MotionMatchingTrait.cpp
+++ b/Engine/Plugins/Experimental/UAF/UAFPoseSearch/Source/UAFPoseSearch/Private/MotionMatchingTrait.cpp
@@ -613,7 +613,7 @@ void FAnimNextMotionMatchingTask::Execute(UE::UAF::FEvaluationVM& VM) const
 	using namespace UE::UAF;
 	using namespace UE::PoseSearch;
 
-	if (VM.GetContext().IsValid())
+	if (VM.GetActiveNamedSet())
 	{
 		// TODO: Implement with new attribute runtime
 		return;
diff --git a/Engine/Plugins/Experimental/UAF/UAFPoseSearch/Source/UAFPoseSearch/Private/PoseSearchResultEmulatorTrait.cpp b/Engine/Plugins/Experimental/UAF/UAFPoseSearch/Source/UAFPoseSearch/Private/PoseSearchResultEmulatorTrait.cpp
--- a/Engine/Plugins/Experimental/UAF/UAFPoseSearch/Source/UAFPoseSearch/Private/PoseSearchResultEmulatorTrait.cpp
+++ b/Engine/Plugins/Experimental/UAF/UAFPoseSearch/Source/UAFPoseSearch/Private/PoseSearchResultEmulatorTrait.cpp
@@ -196,7 +196,7 @@ void FAnimNextPoseSearchResultEmulatorTask::Execute(UE::UAF::FEvaluationVM& VM)
 {
 	using namespace UE::UAF;
 
-	if (VM.GetContext().IsValid())
+	if (VM.GetActiveNamedSet())
 	{
 		// TODO: Implement with new attribute runtime
 		return;
diff --git a/Engine/Plugins/Experimental/UAF/UAFWarping/Source/UAFWarping/Private/SteeringTrait.cpp b/Engine/Plugins/Experimental/UAF/UAFWarping/Source/UAFWarping/Private/SteeringTrait.cpp
--- a/Engine/Plugins/Experimental/UAF/UAFWarping/Source/UAFWarping/Private/SteeringTrait.cpp
+++ b/Engine/Plugins/Experimental/UAF/UAFWarping/Source/UAFWarping/Private/SteeringTrait.cpp
@@ -109,7 +109,7 @@ void FAnimNextSteeringTask::Execute(UE::UAF::FEvaluationVM& VM) const
 {
 	using namespace UE::UAF;
 
-	if (VM.GetContext().IsValid())
+	if (VM.GetActiveNamedSet())
 	{
 		// TODO: Implement with new attribute runtime
 		return;
diff --git a/Engine/Plugins/Experimental/UAF/UAFWarping/Source/UAFWarping/Private/StrafeWarpingTrait.cpp b/Engine/Plugins/Experimental/UAF/UAFWarping/Source/UAFWarping/Private/StrafeWarpingTrait.cpp
--- a/Engine/Plugins/Experimental/UAF/UAFWarping/Source/UAFWarping/Private/StrafeWarpingTrait.cpp
+++ b/Engine/Plugins/Experimental/UAF/UAFWarping/Source/UAFWarping/Private/StrafeWarpingTrait.cpp
@@ -186,7 +186,7 @@ void FAnimNextStrafeWarpingTask::Execute(UE::UAF::FEvaluationVM& VM) const
 	using namespace UE::UAF;
 	QUICK_SCOPE_CYCLE_COUNTER(FAnimNextStrafeWarpingTask_Execute);
 
-	if (VM.GetContext().IsValid())
+	if (VM.GetActiveNamedSet())
 	{
 		// TODO: Implement with new attribute runtime
 		return;
diff --git a/Engine/Plugins/Experimental/UAF/UAFWarping/Source/UAFWarping/Private/WarpTestTrait.cpp b/Engine/Plugins/Experimental/UAF/UAFWarping/Source/UAFWarping/Private/WarpTestTrait.cpp
--- a/Engine/Plugins/Experimental/UAF/UAFWarping/Source/UAFWarping/Private/WarpTestTrait.cpp
+++ b/Engine/Plugins/Experimental/UAF/UAFWarping/Source/UAFWarping/Private/WarpTestTrait.cpp
@@ -95,7 +95,7 @@ void FAnimNextWarpTestTask::Execute(UE::UAF::FEvaluationVM& VM) const
 {
 	using namespace UE::UAF;
 
-	if (VM.GetContext().IsValid())
+	if (VM.GetActiveNamedSet())
 	{
 		// TODO: Implement with new attribute runtime
 		return;
