Expose a pin for the attribute set to use on the RunGraph node

#Synced-CL 48947880

diff --git a/Engine/Plugins/Experimental/UAF/UAFAnimGraph/Source/UAFAnimGraph/Internal/Graph/RigUnit_AnimNextRunAnimationGraph_v2.h b/Engine/Plugins/Experimental/UAF/UAFAnimGraph/Source/UAFAnimGraph/Internal/Graph/RigUnit_AnimNextRunAnimationGraph_v2.h
--- a/Engine/Plugins/Experimental/UAF/UAFAnimGraph/Source/UAFAnimGraph/Internal/Graph/RigUnit_AnimNextRunAnimationGraph_v2.h
+++ b/Engine/Plugins/Experimental/UAF/UAFAnimGraph/Source/UAFAnimGraph/Internal/Graph/RigUnit_AnimNextRunAnimationGraph_v2.h
@@ -52,6 +52,10 @@ struct FRigUnit_AnimNextRunAnimationGraph_v2 : public FRigUnit_AnimNextBase
 	UPROPERTY(EditAnywhere, Category = "Run", meta = (Input))
 	FAnimNextVariableOverridesCollection Overrides;
 
+	// The name of the attribute set to execute with, if 'None' will include everything
+	UPROPERTY(EditAnywhere, Category = "Run", meta = (Input))
+	FName AttributeSet;
+
 	// Pose result
 	UPROPERTY(EditAnywhere, Category = "Run", meta = (Output))
 	FAnimNextGraphLODPose Result;
diff --git a/Engine/Plugins/Experimental/UAF/UAFAnimGraph/Source/UAFAnimGraph/Private/Graph/RigUnit_AnimNextRunAnimationGraph_v2.cpp b/Engine/Plugins/Experimental/UAF/UAFAnimGraph/Source/UAFAnimGraph/Private/Graph/RigUnit_AnimNextRunAnimationGraph_v2.cpp
--- a/Engine/Plugins/Experimental/UAF/UAFAnimGraph/Source/UAFAnimGraph/Private/Graph/RigUnit_AnimNextRunAnimationGraph_v2.cpp
+++ b/Engine/Plugins/Experimental/UAF/UAFAnimGraph/Source/UAFAnimGraph/Private/Graph/RigUnit_AnimNextRunAnimationGraph_v2.cpp
@@ -139,9 +139,8 @@ FRigUnit_AnimNextRunAnimationGraph_v2_Execute()
 		UUAFSkeletonUserData* UserData = UUAFSkeletonUserData::FromSkeleton(*Skeleton);
 
 		UAbstractSkeletonSetBinding* Binding = UserData != nullptr ? UserData->GetSetBinding() : nullptr;
-		FName SetName;
 
-		EvaluateGraphContext.SetSkeletonSet(Binding, SetName);
+		EvaluateGraphContext.SetSkeletonSet(Binding, AttributeSet);
 	}
 
 	EvaluateGraph(EvaluateGraphContext, Result);
diff --git a/Engine/Plugins/Experimental/UAF/UAFAnimGraph/Source/UAFAnimGraph/Private/Native/UAFSimplePrePhysicsGraphComponent.cpp b/Engine/Plugins/Experimental/UAF/UAFAnimGraph/Source/UAFAnimGraph/Private/Native/UAFSimplePrePhysicsGraphComponent.cpp
--- a/Engine/Plugins/Experimental/UAF/UAFAnimGraph/Source/UAFAnimGraph/Private/Native/UAFSimplePrePhysicsGraphComponent.cpp
+++ b/Engine/Plugins/Experimental/UAF/UAFAnimGraph/Source/UAFAnimGraph/Private/Native/UAFSimplePrePhysicsGraphComponent.cpp
@@ -42,6 +42,7 @@ void FUAFSimplePrePhysicsGraphComponent::PrePhysics(TStructView<FUAFScriptCompon
 			LOD,
 			This.RefPose,
 			FAnimNextVariableOverridesCollection(),
+			NAME_None,
 			This.Pose,
 			This.WorkData);
 
