Add layering support for UAF value runtime

#Synced-CL 48953711

diff --git a/Engine/Plugins/Experimental/UAF/UAF/Source/UAF/Private/UAF/Attributes/EngineAttributes.cpp b/Engine/Plugins/Experimental/UAF/UAF/Source/UAF/Private/UAF/Attributes/EngineAttributes.cpp
--- a/Engine/Plugins/Experimental/UAF/UAF/Source/UAF/Private/UAF/Attributes/EngineAttributes.cpp
+++ b/Engine/Plugins/Experimental/UAF/UAF/Source/UAF/Private/UAF/Attributes/EngineAttributes.cpp
@@ -6,6 +6,7 @@
 #include "UAF/Attributes/Transformers/Accumulate.h"
 #include "UAF/Attributes/Transformers/AdditiveSpace.h"
 #include "UAF/Attributes/Transformers/Interpolate.h"
+#include "UAF/Attributes/Transformers/Layer.h"
 #include "UAF/Attributes/Transformers/Overwrite.h"
 #include "UAF/Attributes/Transformers/Sanitize.h"
 
@@ -31,6 +32,8 @@ namespace UE::UAF
 		RuntimeRegistry.RegisterAttributeValueTransformer<Transformers::Private::FApplyAdditiveSpace_FloatAttribute>();
 		RuntimeRegistry.RegisterAttributeValueTransformer<Transformers::Private::FInterpolate_BoneTransformAttribute>();
 		RuntimeRegistry.RegisterAttributeValueTransformer<Transformers::Private::FInterpolate_FloatAttribute>();
+		RuntimeRegistry.RegisterAttributeValueTransformer<Transformers::Private::FLayer_BoneTransformAttribute>();
+		RuntimeRegistry.RegisterAttributeValueTransformer<Transformers::Private::FLayer_FloatAttribute>();
 		RuntimeRegistry.RegisterAttributeValueTransformer<Transformers::Private::FMakeAdditiveSpace_BoneTransformAttribute>();
 		RuntimeRegistry.RegisterAttributeValueTransformer<Transformers::Private::FMakeAdditiveSpace_FloatAttribute>();
 		RuntimeRegistry.RegisterAttributeValueTransformer<Transformers::Private::FOverwrite_BoneTransformAttribute>();
diff --git a/Engine/Plugins/Experimental/UAF/UAF/Source/UAF/Private/UAF/Attributes/Transformers/Interpolate.cpp b/Engine/Plugins/Experimental/UAF/UAF/Source/UAF/Private/UAF/Attributes/Transformers/Interpolate.cpp
--- a/Engine/Plugins/Experimental/UAF/UAF/Source/UAF/Private/UAF/Attributes/Transformers/Interpolate.cpp
+++ b/Engine/Plugins/Experimental/UAF/UAF/Source/UAF/Private/UAF/Attributes/Transformers/Interpolate.cpp
@@ -8,26 +8,49 @@ namespace UE::UAF::Transformers
 {
 	FName FInterpolate::TransformerName = TEXT("UAF::Transformers::Interpolate");
 
-	void FInterpolate::Apply(const FAttributeValueTransformerMapPtr& TransformerMap, const FAttributeCollection& CollectionA, const FAttributeCollection& CollectionB, float Alpha, FAttributeCollection& OutputCollection)
+	void FInterpolate::Apply(
+		const FAttributeValueTransformerMapPtr& TransformerMap,
+		const FAttributeCollection& InputA,
+		const FAttributeCollection& InputB,
+		float WeightB,
+		FAttributeCollection& Output)
 	{
 		// Named sets must match as it ensures that our inputs have the same sizes/shapes
-		checkf(CollectionA.GetNamedSet() == CollectionB.GetNamedSet(), TEXT("Both input collections must use the same named set"));
-		checkf(CollectionA.GetValueSpace() == CollectionB.GetValueSpace(), TEXT("Both input collections must have matching value spaces: [%s] != [%s]"), *CollectionA.GetValueSpace().ToString(), *CollectionB.GetValueSpace().ToString());
+		checkf(InputA.GetNamedSet() == InputB.GetNamedSet(), TEXT("Both input collections must use the same named set"));
+		checkf(InputA.GetValueSpace() == InputB.GetValueSpace(), TEXT("Both input collections must have matching value spaces: [%s] != [%s]"), *InputA.GetValueSpace().ToString(), *InputB.GetValueSpace().ToString());
 
-		const bool bIsOutputA = &CollectionA == &OutputCollection;
-		if (bIsOutputA && !FAnimWeight::IsRelevant(Alpha))
+		const bool bIsOutputA = &InputA == &Output;
+		if (bIsOutputA && !FAnimWeight::IsRelevant(WeightB))
 		{
 			// Our input A is our output and B is not relevant, nothing to do
 			return;
 		}
 
-		const bool bIsOutputB = &CollectionB == &OutputCollection;
-		if (bIsOutputB && FAnimWeight::IsFullWeight(Alpha))
+		const bool bIsOutputB = &InputB == &Output;
+		if (bIsOutputB && FAnimWeight::IsFullWeight(WeightB))
 		{
 			// Our input B is our output and A is not relevant, nothing to do
 			return;
 		}
 
+		// Allocator is irrelevant since it will remain empty
+		FAttributeCollectionHeap EmptyCollection(InputA.GetNamedSet());
+		FInterpolate::Apply(TransformerMap, InputA, InputB, EmptyCollection, WeightB, Output);
+	}
+
+	void FInterpolate::Apply(
+		const FAttributeValueTransformerMapPtr& TransformerMap,
+		const FAttributeCollection& InputA,
+		const FAttributeCollection& InputB,
+		const FAttributeCollection& PerValueWeightsB,
+		float DefaultWeightB,
+		FAttributeCollection& Output)
+	{
+		// Named sets must match as it ensures that our inputs have the same sizes/shapes
+		checkf(InputA.GetNamedSet() == InputB.GetNamedSet(), TEXT("Both input collections must use the same named set"));
+		checkf(InputA.GetNamedSet() == PerValueWeightsB.GetNamedSet(), TEXT("Input and per-value weight collections must use the same named set"));
+		checkf(InputA.GetValueSpace() == InputB.GetValueSpace(), TEXT("Both input collections must have matching value spaces: [%s] != [%s]"), *InputA.GetValueSpace().ToString(), *InputB.GetValueSpace().ToString());
+
 		const FAttributeValueTransformerList* TransformerList = TransformerMap->Find(TransformerName);
 		if (TransformerList == nullptr)
 		{
@@ -35,10 +58,14 @@ namespace UE::UAF::Transformers
 			return;
 		}
 
+		UScriptStruct* FloatValueType = FFloatAnimationAttribute::StaticStruct();
+
+		const bool bIsOutputA = &InputA == &Output;
+		const bool bIsOutputB = &InputB == &Output;
 		if (bIsOutputA || bIsOutputB)
 		{
 			// Our output collection matches one of our inputs, we'll modify it in place
-			FReallocFun OutputAllocator = OutputCollection.GetAllocator();
+			FReallocFun OutputAllocator = Output.GetAllocator();
 
 			// Set maps
 			{
@@ -47,7 +74,7 @@ namespace UE::UAF::Transformers
 
 				OuterJoinBy(
 					FAttributeValueTransformerListWithAttributeCollectionSetMapJoinOp(),
-					[Alpha, bIsOutputA, OutputAllocator, &PendingSetMaps](FRawTransformerFunc Transformer, FAttributeSetMap* SetMapA, FAttributeSetMap* SetMapB)
+					[&PerValueWeightsB, DefaultWeightB, bIsOutputA, FloatValueType, OutputAllocator, &PendingSetMaps](FRawTransformerFunc Transformer, FAttributeSetMap* SetMapA, FAttributeSetMap* SetMapB)
 					{
 						if (SetMapA == nullptr)
 						{
@@ -61,8 +88,7 @@ namespace UE::UAF::Transformers
 								if (bIsOutputA)
 								{
 									// A is our output, duplicate it
-									FAttributeSetMap* SetMapCopy = SetMapB->Duplicate(OutputAllocator);
-									PendingSetMaps.Add(SetMapCopy);
+									PendingSetMaps.Add(SetMapB->Duplicate(OutputAllocator));
 								}
 								else
 								{
@@ -82,8 +108,7 @@ namespace UE::UAF::Transformers
 								else
 								{
 									// B is our output, duplicate it
-									FAttributeSetMap* SetMapCopy = SetMapA->Duplicate(OutputAllocator);
-									PendingSetMaps.Add(SetMapCopy);
+									PendingSetMaps.Add(SetMapA->Duplicate(OutputAllocator));
 								}
 							}
 							else
@@ -91,7 +116,7 @@ namespace UE::UAF::Transformers
 								if (Transformer == nullptr)
 								{
 									// No transformer was found for this value type, retain the one with the largest weight
-									if (Alpha >= 0.5f)
+									if (DefaultWeightB >= 0.5f)
 									{
 										// B has more weight
 										if (bIsOutputA)
@@ -123,21 +148,25 @@ namespace UE::UAF::Transformers
 									// Transform our matching pair in place
 									FAttributeSetMap* OutputSetMap = bIsOutputA ? SetMapA : SetMapB;
 
+									// We cannot join our weights because their value type differs
+									const FAttributeMappingKey MappingKey = SetMapA->GetMappingKey().To(FloatValueType);
+									const FAttributeSetMap* PerValueWeightMapB = PerValueWeightsB.FindSetMap(MappingKey);
+
 									FTransformSetMapFunc TypedTransformer = reinterpret_cast<FTransformSetMapFunc>(Transformer);
-									(*TypedTransformer)(SetMapA, SetMapB, Alpha, OutputSetMap, nullptr);
+									(*TypedTransformer)(SetMapA, SetMapB, PerValueWeightMapB, DefaultWeightB, OutputSetMap, nullptr);
 								}
 							}
 						}
 					},
 					TransformerList->CreateSetMapTransformerIterator(),
 					// Remove the const qualifier since we'll mutate one in place as our output
-					const_cast<FAttributeCollection&>(CollectionA).CreateSetMapIterator(),
-					const_cast<FAttributeCollection&>(CollectionB).CreateSetMapIterator());
+					const_cast<FAttributeCollection&>(InputA).CreateSetMapIterator(),
+					const_cast<FAttributeCollection&>(InputB).CreateSetMapIterator());
 
 				// Append our pending entries
 				for (FAttributeSetMap* SetMap : PendingSetMaps)
 				{
-					OutputCollection.AddSetMap(SetMap);
+					Output.AddSetMap(SetMap);
 				}
 			}
 
@@ -148,7 +177,7 @@ namespace UE::UAF::Transformers
 
 				OuterJoinBy(
 					FAttributeValueTransformerListWithAttributeCollectionMapJoinOp(),
-					[Alpha, bIsOutputA, OutputAllocator, &PendingMaps](FRawTransformerFunc Transformer, FAttributeMap* MapA, FAttributeMap* MapB)
+					[DefaultWeightB, bIsOutputA, OutputAllocator, &PendingMaps](FRawTransformerFunc Transformer, FAttributeMap* MapA, FAttributeMap* MapB)
 					{
 						if (MapA == nullptr)
 						{
@@ -162,8 +191,7 @@ namespace UE::UAF::Transformers
 								if (bIsOutputA)
 								{
 									// A is our output, duplicate it
-									FAttributeMap* MapCopy = MapB->Duplicate(OutputAllocator);
-									PendingMaps.Add(MapCopy);
+									PendingMaps.Add(MapB->Duplicate(OutputAllocator));
 								}
 								else
 								{
@@ -183,8 +211,7 @@ namespace UE::UAF::Transformers
 								else
 								{
 									// B is our output, duplicate it
-									FAttributeMap* MapCopy = MapA->Duplicate(OutputAllocator);
-									PendingMaps.Add(MapCopy);
+									PendingMaps.Add(MapA->Duplicate(OutputAllocator));
 								}
 							}
 							else
@@ -192,7 +219,7 @@ namespace UE::UAF::Transformers
 								if (Transformer == nullptr)
 								{
 									// No transformer was found for this value type, retain the one with the largest weight
-									if (Alpha >= 0.5f)
+									if (DefaultWeightB >= 0.5f)
 									{
 										// B has more weight
 										if (bIsOutputA)
@@ -225,7 +252,7 @@ namespace UE::UAF::Transformers
 									FAttributeMap* OutputMap = nullptr;
 
 									FTransformMapFunc TypedTransformer = reinterpret_cast<FTransformMapFunc>(Transformer);
-									(*TypedTransformer)(MapA, MapB, Alpha, OutputMap, OutputAllocator);
+									(*TypedTransformer)(MapA, MapB, DefaultWeightB, OutputMap, OutputAllocator);
 
 									checkf(OutputMap != nullptr, TEXT("Transformer function is responsible for allocating an output buffer"));
 									checkf(OutputMap->GetAllocator() == OutputAllocator, TEXT("Transformer function is responsible for allocating using the output allocator"));
@@ -241,26 +268,26 @@ namespace UE::UAF::Transformers
 					},
 					TransformerList->CreateMapTransformerIterator(),
 					// Remove the const qualifier since we'll mutate one in place as our output
-					const_cast<FAttributeCollection&>(CollectionA).CreateMapIterator(),
-					const_cast<FAttributeCollection&>(CollectionB).CreateMapIterator());
+					const_cast<FAttributeCollection&>(InputA).CreateMapIterator(),
+					const_cast<FAttributeCollection&>(InputB).CreateMapIterator());
 
 				// Append our pending entries
 				for (FAttributeMap* Map : PendingMaps)
 				{
-					OutputCollection.AddMap(Map);
+					Output.AddMap(Map);
 				}
 			}
 		}
 		else
 		{
 			// Our output collection isn't either of our inputs, reset it
-			OutputCollection.Reset(CollectionA.GetNamedSet());
-			OutputCollection.SetValueSpace(CollectionA.GetValueSpace());
+			Output.Reset(InputA.GetNamedSet());
+			Output.SetValueSpace(InputA.GetValueSpace());
 
 			// Set maps
 			OuterJoinBy(
 				FAttributeValueTransformerListWithAttributeCollectionSetMapJoinOp(),
-				[Alpha, &OutputCollection](FRawTransformerFunc Transformer, const FAttributeSetMap* SetMapA, const FAttributeSetMap* SetMapB)
+				[&PerValueWeightsB, DefaultWeightB, FloatValueType, &Output](FRawTransformerFunc Transformer, const FAttributeSetMap* SetMapA, const FAttributeSetMap* SetMapB)
 				{
 					if (SetMapA == nullptr)
 					{
@@ -271,8 +298,7 @@ namespace UE::UAF::Transformers
 						else
 						{
 							// This set map doesn't have a matching pair, duplicate it
-							FAttributeSetMap* SetMapCopy = SetMapB->Duplicate(OutputCollection.GetAllocator());
-							OutputCollection.AppendSetMap(SetMapCopy);
+							Output.AppendSetMap(SetMapB->Duplicate(Output.GetAllocator()));
 						}
 					}
 					else
@@ -280,8 +306,7 @@ namespace UE::UAF::Transformers
 						if (SetMapB == nullptr)
 						{
 							// This set map doesn't have a matching pair, duplicate it
-							FAttributeSetMap* SetMapCopy = SetMapA->Duplicate(OutputCollection.GetAllocator());
-							OutputCollection.AppendSetMap(SetMapCopy);
+							Output.AppendSetMap(SetMapA->Duplicate(Output.GetAllocator()));
 						}
 						else
 						{
@@ -289,40 +314,44 @@ namespace UE::UAF::Transformers
 							{
 								// No transformer was found for this value type, retain the one with the largest weight
 								FAttributeSetMap* SetMapCopy;
-								if (Alpha >= 0.5f)
+								if (DefaultWeightB >= 0.5f)
 								{
-									SetMapCopy = SetMapB->Duplicate(OutputCollection.GetAllocator());
+									SetMapCopy = SetMapB->Duplicate(Output.GetAllocator());
 								}
 								else
 								{
-									SetMapCopy = SetMapA->Duplicate(OutputCollection.GetAllocator());
+									SetMapCopy = SetMapA->Duplicate(Output.GetAllocator());
 								}
 
-								OutputCollection.AppendSetMap(SetMapCopy);
+								Output.AppendSetMap(SetMapCopy);
 							}
 							else
 							{
 								// Transform our matching pair
 								FAttributeSetMap* OutputSetMap = nullptr;
 
+								// We cannot join our weights because their value type differs
+								const FAttributeMappingKey MappingKey = SetMapA->GetMappingKey().To(FloatValueType);
+								const FAttributeSetMap* PerValueWeightMapB = PerValueWeightsB.FindSetMap(MappingKey);
+
 								FTransformSetMapFunc TypedTransformer = reinterpret_cast<FTransformSetMapFunc>(Transformer);
-								(*TypedTransformer)(SetMapA, SetMapB, Alpha, OutputSetMap, OutputCollection.GetAllocator());
+								(*TypedTransformer)(SetMapA, SetMapB, PerValueWeightMapB, DefaultWeightB, OutputSetMap, Output.GetAllocator());
 
 								checkf(OutputSetMap != nullptr, TEXT("Transformer function is responsible for allocating an output buffer"));
-								checkf(OutputSetMap->GetAllocator() == OutputCollection.GetAllocator(), TEXT("Transformer function is responsible for allocating using the output allocator"));
-								OutputCollection.AppendSetMap(OutputSetMap);
+								checkf(OutputSetMap->GetAllocator() == Output.GetAllocator(), TEXT("Transformer function is responsible for allocating using the output allocator"));
+								Output.AppendSetMap(OutputSetMap);
 							}
 						}
 					}
 				},
 				TransformerList->CreateSetMapTransformerIterator(),
-				CollectionA.CreateConstSetMapIterator(),
-				CollectionB.CreateConstSetMapIterator());
+				InputA.CreateConstSetMapIterator(),
+				InputB.CreateConstSetMapIterator());
 
 			// Name/value maps
 			OuterJoinBy(
 				FAttributeValueTransformerListWithAttributeCollectionMapJoinOp(),
-				[Alpha, &OutputCollection](FRawTransformerFunc Transformer, const FAttributeMap* MapA, const FAttributeMap* MapB)
+				[DefaultWeightB, &Output](FRawTransformerFunc Transformer, const FAttributeMap* MapA, const FAttributeMap* MapB)
 				{
 					if (MapA == nullptr)
 					{
@@ -333,8 +362,7 @@ namespace UE::UAF::Transformers
 						else
 						{
 							// This map doesn't have a matching pair, duplicate it
-							FAttributeMap* MapCopy = MapB->Duplicate(OutputCollection.GetAllocator());
-							OutputCollection.AppendMap(MapCopy);
+							Output.AppendMap(MapB->Duplicate(Output.GetAllocator()));
 						}
 					}
 					else
@@ -350,16 +378,16 @@ namespace UE::UAF::Transformers
 							{
 								// No transformer was found for this value type, retain the one with the largest weight
 								FAttributeMap* MapCopy;
-								if (Alpha >= 0.5f)
+								if (DefaultWeightB >= 0.5f)
 								{
-									MapCopy = MapB->Duplicate(OutputCollection.GetAllocator());
+									MapCopy = MapB->Duplicate(Output.GetAllocator());
 								}
 								else
 								{
-									MapCopy = MapA->Duplicate(OutputCollection.GetAllocator());
+									MapCopy = MapA->Duplicate(Output.GetAllocator());
 								}
 
-								OutputCollection.AppendMap(MapCopy);
+								Output.AppendMap(MapCopy);
 							}
 							else
 							{
@@ -367,24 +395,24 @@ namespace UE::UAF::Transformers
 								FAttributeMap* OutputMap = nullptr;
 
 								FTransformMapFunc TypedTransformer = reinterpret_cast<FTransformMapFunc>(Transformer);
-								(*TypedTransformer)(MapA, MapB, Alpha, OutputMap, OutputCollection.GetAllocator());
+								(*TypedTransformer)(MapA, MapB, DefaultWeightB, OutputMap, Output.GetAllocator());
 
 								checkf(OutputMap != nullptr, TEXT("Transformer function is responsible for allocating an output buffer"));
-								checkf(OutputMap->GetAllocator() == OutputCollection.GetAllocator(), TEXT("Transformer function is responsible for allocating using the output allocator"));
-								OutputCollection.AppendMap(OutputMap);
+								checkf(OutputMap->GetAllocator() == Output.GetAllocator(), TEXT("Transformer function is responsible for allocating using the output allocator"));
+								Output.AppendMap(OutputMap);
 							}
 						}
 					}
 				},
 				TransformerList->CreateMapTransformerIterator(),
-				CollectionA.CreateConstMapIterator(),
-				CollectionB.CreateConstMapIterator());
+				InputA.CreateConstMapIterator(),
+				InputB.CreateConstMapIterator());
 		}
 	}
 
 	namespace Private
 	{
-		void FInterpolate_BoneTransformAttribute::TransformSetMap(const FAttributeSetMap* SetMapA, const FAttributeSetMap* SetMapB, float Alpha, FAttributeSetMap*& OutputSetMap, FReallocFun OutputAllocator)
+		void FInterpolate_BoneTransformAttribute::TransformSetMap(const FAttributeSetMap* SetMapA, const FAttributeSetMap* SetMapB, const FAttributeSetMap* PerValueWeightMapB, float DefaultWeightB, FAttributeSetMap*& OutputSetMap, FReallocFun OutputAllocator)
 		{
 			check(SetMapA->GetTypedSet() == SetMapB->GetTypedSet());
 			check(SetMapA->Num() == SetMapB->Num());
@@ -407,17 +435,37 @@ namespace UE::UAF::Transformers
 			FBoneTransformAnimationAttribute* OutputTransformPtr = OutputBoneTransforms->GetData();
 			const FBoneTransformAnimationAttribute* TransformPtrAEnd = TransformPtrA + BoneTransformsA->Num();
 
-			while (TransformPtrA < TransformPtrAEnd)
+			if (const TAttributeSetMap<FFloatAnimationAttribute>* PerValueWeightsB = Cast<FFloatAnimationAttribute>(PerValueWeightMapB))
+			{
+				check(SetMapA->GetTypedSet() == PerValueWeightsB->GetTypedSet());
+				check(SetMapA->Num() == PerValueWeightsB->Num());
+
+				const FFloatAnimationAttribute* PerValueWeightsBPtr = PerValueWeightsB->GetData();
+
+				while (TransformPtrA < TransformPtrAEnd)
+				{
+					OutputTransformPtr->Value.Blend(TransformPtrA->Value, TransformPtrB->Value, PerValueWeightsBPtr->Value);
+
+					TransformPtrA++;
+					TransformPtrB++;
+					PerValueWeightsBPtr++;
+					OutputTransformPtr++;
+				}
+			}
+			else
 			{
-				OutputTransformPtr->Value.Blend(TransformPtrA->Value, TransformPtrB->Value, Alpha);
+				while (TransformPtrA < TransformPtrAEnd)
+				{
+					OutputTransformPtr->Value.Blend(TransformPtrA->Value, TransformPtrB->Value, DefaultWeightB);
 
-				TransformPtrA++;
-				TransformPtrB++;
-				OutputTransformPtr++;
+					TransformPtrA++;
+					TransformPtrB++;
+					OutputTransformPtr++;
+				}
 			}
 		}
 
-		void FInterpolate_BoneTransformAttribute::TransformMap(const FAttributeMap* MapA, const FAttributeMap* MapB, float Alpha, FAttributeMap*& OutputMap, FReallocFun OutputAllocator)
+		void FInterpolate_BoneTransformAttribute::TransformMap(const FAttributeMap* MapA, const FAttributeMap* MapB, float WeightB, FAttributeMap*& OutputMap, FReallocFun OutputAllocator)
 		{
 			check(OutputMap == nullptr);
 
@@ -427,14 +475,14 @@ namespace UE::UAF::Transformers
 
 			OuterJoinBy(
 				FAttributeMapJoinOp(),
-				[&Alpha, OutputBoneTransforms](FName Name, const FBoneTransformAnimationAttribute* A, const FBoneTransformAnimationAttribute* B)
+				[WeightB, OutputBoneTransforms](FName Name, const FBoneTransformAnimationAttribute* A, const FBoneTransformAnimationAttribute* B)
 				{
 					// Missing values are set to the identity
 					const FTransform& ValueA = A != nullptr ? A->Value : FTransform::Identity;
 					const FTransform& ValueB = B != nullptr ? B->Value : FTransform::Identity;
 
 					FBoneTransformAnimationAttribute Output;
-					Output.Value.Blend(ValueA, ValueB, Alpha);
+					Output.Value.Blend(ValueA, ValueB, WeightB);
 
 					OutputBoneTransforms->Append(Name, Output);
 				},
@@ -444,7 +492,7 @@ namespace UE::UAF::Transformers
 			OutputMap = OutputBoneTransforms;
 		}
 
-		void FInterpolate_FloatAttribute::TransformSetMap(const FAttributeSetMap* SetMapA, const FAttributeSetMap* SetMapB, float Alpha, FAttributeSetMap*& OutputSetMap, FReallocFun OutputAllocator)
+		void FInterpolate_FloatAttribute::TransformSetMap(const FAttributeSetMap* SetMapA, const FAttributeSetMap* SetMapB, const FAttributeSetMap* PerValueWeightMapB, float DefaultWeightB, FAttributeSetMap*& OutputSetMap, FReallocFun OutputAllocator)
 		{
 			check(SetMapA->GetTypedSet() == SetMapB->GetTypedSet());
 			check(SetMapA->Num() == SetMapB->Num());
@@ -467,17 +515,37 @@ namespace UE::UAF::Transformers
 			FFloatAnimationAttribute* OutputFloatsPtr = OutputFloats->GetData();
 			const FFloatAnimationAttribute* FloatsPtrAEnd = FloatsPtrA + FloatsA->Num();
 
-			while (FloatsPtrA < FloatsPtrAEnd)
+			if (const TAttributeSetMap<FFloatAnimationAttribute>* PerValueWeightsB = Cast<FFloatAnimationAttribute>(PerValueWeightMapB))
 			{
-				OutputFloatsPtr->Value = FMath::Lerp(FloatsPtrA->Value, FloatsPtrB->Value, Alpha);
+				check(SetMapA->GetTypedSet() == PerValueWeightsB->GetTypedSet());
+				check(SetMapA->Num() == PerValueWeightsB->Num());
+
+				const FFloatAnimationAttribute* PerValueWeightsBPtr = PerValueWeightsB->GetData();
+
+				while (FloatsPtrA < FloatsPtrAEnd)
+				{
+					OutputFloatsPtr->Value = FMath::Lerp(FloatsPtrA->Value, FloatsPtrB->Value, PerValueWeightsBPtr->Value);
 
-				FloatsPtrA++;
-				FloatsPtrB++;
-				OutputFloatsPtr++;
+					FloatsPtrA++;
+					FloatsPtrB++;
+					PerValueWeightsBPtr++;
+					OutputFloatsPtr++;
+				}
+			}
+			else
+			{
+				while (FloatsPtrA < FloatsPtrAEnd)
+				{
+					OutputFloatsPtr->Value = FMath::Lerp(FloatsPtrA->Value, FloatsPtrB->Value, DefaultWeightB);
+
+					FloatsPtrA++;
+					FloatsPtrB++;
+					OutputFloatsPtr++;
+				}
 			}
 		}
 
-		void FInterpolate_FloatAttribute::TransformMap(const FAttributeMap* MapA, const FAttributeMap* MapB, float Alpha, FAttributeMap*& OutputMap, FReallocFun OutputAllocator)
+		void FInterpolate_FloatAttribute::TransformMap(const FAttributeMap* MapA, const FAttributeMap* MapB, float WeightB, FAttributeMap*& OutputMap, FReallocFun OutputAllocator)
 		{
 			check(OutputMap == nullptr);
 
@@ -487,14 +555,14 @@ namespace UE::UAF::Transformers
 
 			OuterJoinBy(
 				FAttributeMapJoinOp(),
-				[&Alpha, OutputFloats](FName Name, const FFloatAnimationAttribute* A, const FFloatAnimationAttribute* B)
+				[WeightB, OutputFloats](FName Name, const FFloatAnimationAttribute* A, const FFloatAnimationAttribute* B)
 				{
 					// Missing values are set to 0.0
 					const float ValueA = A != nullptr ? A->Value : 0.0f;
 					const float ValueB = B != nullptr ? B->Value : 0.0f;
 
 					FFloatAnimationAttribute Output;
-					Output.Value = FMath::Lerp(ValueA, ValueB, Alpha);
+					Output.Value = FMath::Lerp(ValueA, ValueB, WeightB);
 
 					OutputFloats->Append(Name, Output);
 				},
diff --git a/Engine/Plugins/Experimental/UAF/UAF/Source/UAF/Private/UAF/Attributes/Transformers/Layer.cpp b/Engine/Plugins/Experimental/UAF/UAF/Source/UAF/Private/UAF/Attributes/Transformers/Layer.cpp
new file mode 100644
--- /dev/null
+++ b/Engine/Plugins/Experimental/UAF/UAF/Source/UAF/Private/UAF/Attributes/Transformers/Layer.cpp
@@ -0,0 +1,541 @@
+// Copyright Epic Games, Inc. All Rights Reserved.
+
+#include "UAF/Attributes/Transformers/Layer.h"
+
+#include "UAF/Attributes/IteratorUtils.h"
+
+namespace UE::UAF::Transformers
+{
+	FName FLayer::TransformerName = TEXT("UAF::Transformers::Layer");
+
+	void FLayer::Apply(
+		const FAttributeValueTransformerMapPtr& TransformerMap,
+		const FAttributeCollection& Base,
+		const FAttributeCollection& Layer,
+		float LayerWeight,
+		FAttributeCollection& Output)
+	{
+		const bool bIsOutputBase = &Base == &Output;
+		if (bIsOutputBase && !FAnimWeight::IsRelevant(LayerWeight))
+		{
+			// Our input base is our output and the layer is not relevant, nothing to do
+			return;
+		}
+
+		// Allocator is irrelevant since it will remain empty
+		FAttributeCollectionHeap EmptyCollection(Base.GetNamedSet());
+		FLayer::Apply(TransformerMap, Base, Layer, EmptyCollection, LayerWeight, Output);
+	}
+
+	void FLayer::Apply(
+		const FAttributeValueTransformerMapPtr& TransformerMap,
+		const FAttributeCollection& Base,
+		const FAttributeCollection& Layer,
+		const FAttributeCollection& PerValueLayerWeights,
+		float DefaultLayerWeight,
+		FAttributeCollection& Output)
+	{
+		const FAttributeValueTransformerList* TransformerList = TransformerMap->Find(TransformerName);
+		if (TransformerList == nullptr)
+		{
+			// No transformer implementations have been registered
+			return;
+		}
+
+		const float LayerWeight = FMath::Clamp(DefaultLayerWeight, 0.0f, 1.0f);
+
+		FAttributeCollectionStack LayerCopy;
+		const FAttributeCollection* LayerPtr = &Layer;
+
+		const bool bIsOutputBase = &Base == &Output;
+		bool bIsOutputLayer = &Layer == &Output;
+		if (bIsOutputLayer)
+		{
+			// Our output is the layer, we can only modify it in place if it has the same set as our base
+			if (Base.GetNamedSet() == Layer.GetNamedSet())
+			{
+				// We'll modify it in place
+			}
+			else
+			{
+				// Duplicate our layer
+				LayerCopy = Layer;
+				LayerPtr = &LayerCopy;
+				bIsOutputLayer = false;
+			}
+		}
+
+		UScriptStruct* FloatValueType = FFloatAnimationAttribute::StaticStruct();
+
+		if (bIsOutputBase || bIsOutputLayer)
+		{
+			// Our output collection matches one of our inputs, we'll modify it in place
+			FReallocFun OutputAllocator = Output.GetAllocator();
+
+			// Set maps
+			OuterJoinBy(
+				FAttributeValueTransformerListWithAttributeCollectionSetMapJoinOp(),
+				[&PerValueLayerWeights, LayerWeight, bIsOutputBase, FloatValueType, OutputAllocator](FRawTransformerFunc Transformer, FAttributeSetMap* BaseMap, FAttributeSetMap* LayerMap)
+				{
+					if (BaseMap == nullptr)
+					{
+						if (LayerMap == nullptr)
+						{
+							// No maps found for the transformer type
+						}
+						else
+						{
+							// Base doesn't have a matching pair, there is nothing to layer on
+						}
+					}
+					else
+					{
+						if (LayerMap == nullptr)
+						{
+							// Layer doesn't have a matching pair, there is nothing to layer with
+						}
+						else
+						{
+							if (Transformer == nullptr)
+							{
+								// No transformer was found for this value type, retain the one with the largest weight
+								if (LayerWeight >= 0.5f)
+								{
+									// Layer has more weight
+									if (bIsOutputBase)
+									{
+										// Base is our output, copy it over
+										LayerMap->CopyTo(*BaseMap);
+									}
+									else
+									{
+										// Layer is our output, nothing to do
+									}
+								}
+								else
+								{
+									// Base has more weight
+									if (bIsOutputBase)
+									{
+										// Base is our output, nothing to do
+									}
+									else
+									{
+										// Layer is our output, copy it over
+										BaseMap->CopyTo(*LayerMap);
+									}
+								}
+							}
+							else
+							{
+								// Transform our matching pair in place
+								FAttributeSetMap* OutputSetMap = bIsOutputBase ? BaseMap : LayerMap;
+
+								// We cannot join our weights because their value type differs
+								const FAttributeMappingKey MappingKey = BaseMap->GetMappingKey().To(FloatValueType);
+								const FAttributeSetMap* PerValueWeightMap = PerValueLayerWeights.FindSetMap(MappingKey);
+
+								FTransformSetMapFunc TypedTransformer = reinterpret_cast<FTransformSetMapFunc>(Transformer);
+								(*TypedTransformer)(BaseMap, LayerMap, PerValueWeightMap, LayerWeight, OutputSetMap, nullptr);
+							}
+						}
+					}
+				},
+				TransformerList->CreateSetMapTransformerIterator(),
+				// Remove the const qualifier since we'll mutate one in place as our output
+				const_cast<FAttributeCollection&>(Base).CreateSetMapIterator(),
+				const_cast<FAttributeCollection&>(Layer).CreateSetMapIterator());
+
+			// Name/value maps
+			OuterJoinBy(
+				FAttributeValueTransformerListWithAttributeCollectionMapJoinOp(),
+				[LayerWeight, bIsOutputBase, OutputAllocator](FRawTransformerFunc Transformer, FAttributeMap* BaseMap, FAttributeMap* LayerMap)
+				{
+					if (BaseMap == nullptr)
+					{
+						if (LayerMap == nullptr)
+						{
+							// No maps found for the transformer type
+						}
+						else
+						{
+							// Base doesn't have a matching pair, there is nothing to layer on
+						}
+					}
+					else
+					{
+						if (LayerMap == nullptr)
+						{
+							// Layer doesn't have a matching pair, there is nothing to layer with
+						}
+						else
+						{
+							if (Transformer == nullptr)
+							{
+								// No transformer was found for this value type, retain the one with the largest weight
+								if (LayerWeight >= 0.5f)
+								{
+									// B has more weight
+									if (bIsOutputBase)
+									{
+										// A is our output, copy it over
+										LayerMap->CopyTo(*BaseMap);
+									}
+									else
+									{
+										// B is our output, nothing to do
+									}
+								}
+								else
+								{
+									// A has more weight
+									if (bIsOutputBase)
+									{
+										// A is our output, nothing to do
+									}
+									else
+									{
+										// B is our output, copy it over
+										BaseMap->CopyTo(*LayerMap);
+									}
+								}
+							}
+							else
+							{
+								// Transform our matching pair in a new map
+								FAttributeMap* OutputMap = nullptr;
+
+								FTransformMapFunc TypedTransformer = reinterpret_cast<FTransformMapFunc>(Transformer);
+								(*TypedTransformer)(BaseMap, LayerMap, LayerWeight, OutputMap, OutputAllocator);
+
+								checkf(OutputMap != nullptr, TEXT("Transformer function is responsible for allocating an output buffer"));
+								checkf(OutputMap->GetAllocator() == OutputAllocator, TEXT("Transformer function is responsible for allocating using the output allocator"));
+
+								// Overwrite our old output
+								OutputMap->MoveTo(bIsOutputBase ? *BaseMap : *LayerMap);
+
+								// Destroy our temporary output map
+								ReleaseAttributeMap(OutputMap);
+							}
+						}
+					}
+				},
+				TransformerList->CreateMapTransformerIterator(),
+				// Remove the const qualifier since we'll mutate one in place as our output
+				const_cast<FAttributeCollection&>(Base).CreateMapIterator(),
+				const_cast<FAttributeCollection&>(Layer).CreateMapIterator());
+		}
+		else
+		{
+			// Our output collection isn't either of our inputs, reset it
+			Output.Reset(Base.GetNamedSet());
+			Output.SetValueSpace(Base.GetValueSpace());
+
+			// Set maps
+			OuterJoinBy(
+				FAttributeValueTransformerListWithAttributeCollectionSetMapJoinOp(),
+				[&PerValueLayerWeights, LayerWeight, FloatValueType, &Output](FRawTransformerFunc Transformer, const FAttributeSetMap* BaseMap, const FAttributeSetMap* LayerMap)
+				{
+					if (BaseMap == nullptr)
+					{
+						if (LayerMap == nullptr)
+						{
+							// No maps found for the transformer type
+						}
+						else
+						{
+							// This set map doesn't have a matching pair, duplicate it
+							Output.AppendSetMap(LayerMap->Duplicate(Output.GetAllocator()));
+						}
+					}
+					else
+					{
+						if (LayerMap == nullptr)
+						{
+							// This set map doesn't have a matching pair, duplicate it
+							Output.AppendSetMap(BaseMap->Duplicate(Output.GetAllocator()));
+						}
+						else
+						{
+							if (Transformer == nullptr)
+							{
+								// No transformer was found for this value type, retain the one with the largest weight
+								FAttributeSetMap* SetMapCopy;
+								if (LayerWeight >= 0.5f)
+								{
+									SetMapCopy = LayerMap->Duplicate(Output.GetAllocator());
+								}
+								else
+								{
+									SetMapCopy = BaseMap->Duplicate(Output.GetAllocator());
+								}
+
+								Output.AppendSetMap(SetMapCopy);
+							}
+							else
+							{
+								// Transform our matching pair
+								FAttributeSetMap* OutputSetMap = nullptr;
+
+								// We cannot join our weights because their value type differs
+								const FAttributeMappingKey MappingKey = BaseMap->GetMappingKey().To(FloatValueType);
+								const FAttributeSetMap* PerValueWeightMap = PerValueLayerWeights.FindSetMap(MappingKey);
+
+								FTransformSetMapFunc TypedTransformer = reinterpret_cast<FTransformSetMapFunc>(Transformer);
+								(*TypedTransformer)(BaseMap, LayerMap, PerValueWeightMap, LayerWeight, OutputSetMap, Output.GetAllocator());
+
+								checkf(OutputSetMap != nullptr, TEXT("Transformer function is responsible for allocating an output buffer"));
+								checkf(OutputSetMap->GetAllocator() == Output.GetAllocator(), TEXT("Transformer function is responsible for allocating using the output allocator"));
+								Output.AppendSetMap(OutputSetMap);
+							}
+						}
+					}
+				},
+				TransformerList->CreateSetMapTransformerIterator(),
+				Base.CreateConstSetMapIterator(),
+				LayerPtr->CreateConstSetMapIterator());
+
+			// Name/value maps
+			OuterJoinBy(
+				FAttributeValueTransformerListWithAttributeCollectionMapJoinOp(),
+				[LayerWeight, &Output](FRawTransformerFunc Transformer, const FAttributeMap* BaseMap, const FAttributeMap* LayerMap)
+				{
+					if (BaseMap == nullptr)
+					{
+						if (LayerMap == nullptr)
+						{
+							// No maps found for the transformer type
+						}
+						else
+						{
+							// This map doesn't have a matching pair, duplicate it
+							Output.AppendMap(LayerMap->Duplicate(Output.GetAllocator()));
+						}
+					}
+					else
+					{
+						if (LayerMap == nullptr)
+						{
+							// This map doesn't have a matching pair
+							// Our output is A and thus already contains this
+						}
+						else
+						{
+							if (Transformer == nullptr)
+							{
+								// No transformer was found for this value type, retain the one with the largest weight
+								FAttributeMap* MapCopy;
+								if (LayerWeight >= 0.5f)
+								{
+									MapCopy = LayerMap->Duplicate(Output.GetAllocator());
+								}
+								else
+								{
+									MapCopy = BaseMap->Duplicate(Output.GetAllocator());
+								}
+
+								Output.AppendMap(MapCopy);
+							}
+							else
+							{
+								// Transform our matching pair in place
+								FAttributeMap* OutputMap = nullptr;
+
+								FTransformMapFunc TypedTransformer = reinterpret_cast<FTransformMapFunc>(Transformer);
+								(*TypedTransformer)(BaseMap, LayerMap, LayerWeight, OutputMap, Output.GetAllocator());
+
+								checkf(OutputMap != nullptr, TEXT("Transformer function is responsible for allocating an output buffer"));
+								checkf(OutputMap->GetAllocator() == Output.GetAllocator(), TEXT("Transformer function is responsible for allocating using the output allocator"));
+								Output.AppendMap(OutputMap);
+							}
+						}
+					}
+				},
+				TransformerList->CreateMapTransformerIterator(),
+				Base.CreateConstMapIterator(),
+				LayerPtr->CreateConstMapIterator());
+		}
+	}
+
+	namespace Private
+	{
+		void FLayer_BoneTransformAttribute::TransformSetMap(const FAttributeSetMap* Base, const FAttributeSetMap* Layer, const FAttributeSetMap* PerValueWeights, float DefaultLayerWeight, FAttributeSetMap*& OutputSetMap, FReallocFun OutputAllocator)
+		{
+			if (OutputSetMap == nullptr)
+			{
+				const FAttributeSetMap::FConstructArgs Args(Base->GetTypedSet(), Base->GetValueType(), OutputAllocator);
+				OutputSetMap = MakeAttributeSetMap<FBoneTransformAnimationAttribute>(Args);
+			}
+
+			// Our output should have the same shape as our base
+			check(Base->GetTypedSet() == OutputSetMap->GetTypedSet());
+			check(Base->Num() == OutputSetMap->Num());
+
+			// Our per-value weights should have the same shape as our base
+			check(PerValueWeights == nullptr || Base->GetTypedSet() == PerValueWeights->GetTypedSet());
+			check(PerValueWeights == nullptr || Base->Num() == PerValueWeights->Num());
+
+			const FAttributeTypedSetPtr& BaseTypedSet = Base->GetTypedSet();
+			const FAttributeTypedSetPtr& LayerTypedSet = Layer->GetTypedSet();
+			const FAttributeTypedSet* PerValueWeightTypedSet = PerValueWeights != nullptr ? PerValueWeights->GetTypedSet().GetReference() : nullptr;
+
+			const TAttributeSetMap<FBoneTransformAnimationAttribute>* BaseBoneTransforms = CastChecked<FBoneTransformAnimationAttribute>(Base);
+			const TAttributeSetMap<FBoneTransformAnimationAttribute>* LayerBoneTransforms = CastChecked<FBoneTransformAnimationAttribute>(Layer);
+			const TAttributeSetMap<FFloatAnimationAttribute>* PerValueWeightFloats = PerValueWeights != nullptr ? CastChecked<FFloatAnimationAttribute>(PerValueWeights) : nullptr;
+			TAttributeSetMap<FBoneTransformAnimationAttribute>* OutputBoneTransforms = CastChecked<FBoneTransformAnimationAttribute>(OutputSetMap);
+
+			// Use an outer join over our base, layer, and weights
+			// If a value is in the base but not the layer, use the base value
+			// If a value is not in the base but in the layer, ignore it
+			// If a value is in both, blend it
+			//
+			// Efficiency of an outer join wouldn't be great and would be difficult to port so ISPC
+			// Instead we iterate over the base set and use the binding index to get layer values and weights
+
+			const int32 NumBaseBones = BaseTypedSet->Num();
+			for (FAttributeSetIndex BaseSetIndex(0); BaseSetIndex < NumBaseBones; ++BaseSetIndex)
+			{
+				const FAttributeBindingIndex BaseBindingIndex = BaseTypedSet->GetBindingIndex(BaseSetIndex);
+				check(BaseBindingIndex.IsValid());
+
+				FTransform Result = (*BaseBoneTransforms)[BaseSetIndex].Value;
+
+				const FAttributeSetIndex LayerSetIndex = LayerTypedSet->GetIndex(BaseBindingIndex);
+				if (LayerSetIndex.IsValid())
+				{
+					float ValueWeight = DefaultLayerWeight;
+					if (PerValueWeightTypedSet != nullptr)
+					{
+						const FAttributeSetIndex PerWeightSetIndex = PerValueWeightTypedSet->GetIndex(BaseBindingIndex);
+						if (PerWeightSetIndex.IsValid())
+						{
+							ValueWeight = (*PerValueWeightFloats)[PerWeightSetIndex].Value;
+							ValueWeight = FMath::Clamp(ValueWeight, 0.0f, 1.0f);
+						}
+					}
+
+					Result.BlendWith((*LayerBoneTransforms)[LayerSetIndex].Value, ValueWeight);
+				}
+
+				(*OutputBoneTransforms)[BaseSetIndex].Value = Result;
+			}
+		}
+
+		void FLayer_BoneTransformAttribute::TransformMap(const FAttributeMap* Base, const FAttributeMap* Layer, float LayerWeight, FAttributeMap*& OutputMap, FReallocFun OutputAllocator)
+		{
+			check(OutputMap == nullptr);
+
+			const TAttributeMap<FBoneTransformAnimationAttribute>* BaseBoneTransforms = CastChecked<FBoneTransformAnimationAttribute>(Base);
+			const TAttributeMap<FBoneTransformAnimationAttribute>* LayerBoneTransforms = CastChecked<FBoneTransformAnimationAttribute>(Layer);
+
+			TAttributeMap<FBoneTransformAnimationAttribute>* OutputBoneTransforms = MakeAttributeMap<FBoneTransformAnimationAttribute>(OutputAllocator);
+
+			OuterJoinBy(
+				FAttributeMapJoinOp(),
+				[LayerWeight, OutputBoneTransforms](FName Name, const FBoneTransformAnimationAttribute* BaseValue, const FBoneTransformAnimationAttribute* LayerValue)
+				{
+					// Mismatched pairs are ignored
+					if (BaseValue != nullptr && LayerValue != nullptr)
+					{
+						FBoneTransformAnimationAttribute Output = *BaseValue;
+						Output.Value.BlendWith(LayerValue->Value, LayerWeight);
+
+						OutputBoneTransforms->Append(Name, Output);
+					}
+				},
+				BaseBoneTransforms->CreateConstIterator(),
+				LayerBoneTransforms->CreateConstIterator());
+
+			OutputMap = OutputBoneTransforms;
+		}
+
+		void FLayer_FloatAttribute::TransformSetMap(const FAttributeSetMap* Base, const FAttributeSetMap* Layer, const FAttributeSetMap* PerValueWeights, float DefaultLayerWeight, FAttributeSetMap*& OutputSetMap, FReallocFun OutputAllocator)
+		{
+			if (OutputSetMap == nullptr)
+			{
+				const FAttributeSetMap::FConstructArgs Args(Base->GetTypedSet(), Base->GetValueType(), OutputAllocator);
+				OutputSetMap = MakeAttributeSetMap<FFloatAnimationAttribute>(Args);
+			}
+
+			// Our output should have the same shape as our base
+			check(Base->GetTypedSet() == OutputSetMap->GetTypedSet());
+			check(Base->Num() == OutputSetMap->Num());
+
+			// Our per-value weights should have the same shape as our base
+			check(PerValueWeights == nullptr || Base->GetTypedSet() == PerValueWeights->GetTypedSet());
+			check(PerValueWeights == nullptr || Base->Num() == PerValueWeights->Num());
+
+			const FAttributeTypedSetPtr& BaseTypedSet = Base->GetTypedSet();
+			const FAttributeTypedSetPtr& LayerTypedSet = Layer->GetTypedSet();
+			const FAttributeTypedSet* PerValueWeightTypedSet = PerValueWeights != nullptr ? PerValueWeights->GetTypedSet().GetReference() : nullptr;
+
+			const TAttributeSetMap<FFloatAnimationAttribute>* BaseFloats = CastChecked<FFloatAnimationAttribute>(Base);
+			const TAttributeSetMap<FFloatAnimationAttribute>* LayerFloats = CastChecked<FFloatAnimationAttribute>(Layer);
+			const TAttributeSetMap<FFloatAnimationAttribute>* PerValueWeightFloats = PerValueWeights != nullptr ? CastChecked<FFloatAnimationAttribute>(PerValueWeights) : nullptr;
+			TAttributeSetMap<FFloatAnimationAttribute>* OutputFloats = CastChecked<FFloatAnimationAttribute>(OutputSetMap);
+
+			// Use an outer join over our base, layer, and weights
+			// If a value is in the base but not the layer, use the base value
+			// If a value is not in the base but in the layer, ignore it
+			// If a value is in both, blend it
+			//
+			// Efficiency of an outer join wouldn't be great and would be difficult to port so ISPC
+			// Instead we iterate over the base set and use the binding index to get layer values and weights
+
+			const int32 NumBaseBones = BaseTypedSet->Num();
+			for (FAttributeSetIndex BaseSetIndex(0); BaseSetIndex < NumBaseBones; ++BaseSetIndex)
+			{
+				const FAttributeBindingIndex BaseBindingIndex = BaseTypedSet->GetBindingIndex(BaseSetIndex);
+				check(BaseBindingIndex.IsValid());
+
+				float Result = (*BaseFloats)[BaseSetIndex].Value;
+
+				const FAttributeSetIndex LayerSetIndex = LayerTypedSet->GetIndex(BaseBindingIndex);
+				if (LayerSetIndex.IsValid())
+				{
+					float ValueWeight = DefaultLayerWeight;
+					if (PerValueWeightTypedSet != nullptr)
+					{
+						const FAttributeSetIndex PerWeightSetIndex = PerValueWeightTypedSet->GetIndex(BaseBindingIndex);
+						if (PerWeightSetIndex.IsValid())
+						{
+							ValueWeight = (*PerValueWeightFloats)[PerWeightSetIndex].Value;
+							ValueWeight = FMath::Clamp(ValueWeight, 0.0f, 1.0f);
+						}
+					}
+
+					const float LayerValue = (*LayerFloats)[LayerSetIndex].Value;
+					Result = FMath::Lerp(Result, LayerValue, ValueWeight);
+				}
+
+				(*OutputFloats)[BaseSetIndex].Value = Result;
+			}
+		}
+
+		void FLayer_FloatAttribute::TransformMap(const FAttributeMap* Base, const FAttributeMap* Layer, float LayerWeight, FAttributeMap*& OutputMap, FReallocFun OutputAllocator)
+		{
+			check(OutputMap == nullptr);
+
+			const TAttributeMap<FFloatAnimationAttribute>* BaseFloats = CastChecked<FFloatAnimationAttribute>(Base);
+			const TAttributeMap<FFloatAnimationAttribute>* LayerFloats = CastChecked<FFloatAnimationAttribute>(Layer);
+
+			TAttributeMap<FFloatAnimationAttribute>* OutputFloats = MakeAttributeMap<FFloatAnimationAttribute>(OutputAllocator);
+
+			OuterJoinBy(
+				FAttributeMapJoinOp(),
+				[LayerWeight, OutputFloats](FName Name, const FFloatAnimationAttribute* BaseValue, const FFloatAnimationAttribute* LayerValue)
+				{
+					// Mismatched pairs are ignored
+					if (BaseValue != nullptr && LayerValue != nullptr)
+					{
+						FFloatAnimationAttribute Output{ FMath::Lerp(BaseValue->Value, LayerValue->Value, LayerWeight) };
+						OutputFloats->Append(Name, Output);
+					}
+				},
+				BaseFloats->CreateConstIterator(),
+				LayerFloats->CreateConstIterator());
+
+			OutputMap = OutputFloats;
+		}
+	}
+}
diff --git a/Engine/Plugins/Experimental/UAF/UAF/Source/UAF/Public/UAF/Attributes/Transformers/Interpolate.h b/Engine/Plugins/Experimental/UAF/UAF/Source/UAF/Public/UAF/Attributes/Transformers/Interpolate.h
--- a/Engine/Plugins/Experimental/UAF/UAF/Source/UAF/Public/UAF/Attributes/Transformers/Interpolate.h
+++ b/Engine/Plugins/Experimental/UAF/UAF/Source/UAF/Public/UAF/Attributes/Transformers/Interpolate.h
@@ -18,13 +18,23 @@ namespace UE::UAF::Transformers
 		static UE_API FName TransformerName;
 
 		// If the output buffer is nullptr, an allocator is provided and the implementation is responsible for allocating the output buffer and returning it
-		using FTransformSetMapFunc = void (*)(const FAttributeSetMap* SetMapA, const FAttributeSetMap* SetMapB, float Alpha, FAttributeSetMap*& OutputSetMap, FReallocFun OutputAllocator);
-		using FTransformMapFunc = void (*)(const FAttributeMap* MapA, const FAttributeMap* MapB, float Alpha, FAttributeMap*& OutputMap, FReallocFun OutputAllocator);
+		// Note that the input/output buffers can be the same buffer
+		// Per-value map is optional, when not present, the default weight should be used
+		// Unbound values do not support per-value weights
+		using FTransformSetMapFunc = void (*)(const FAttributeSetMap* SetMapA, const FAttributeSetMap* SetMapB, const FAttributeSetMap* PerValueWeightMapB, float DefaultWeightB, FAttributeSetMap*& OutputSetMap, FReallocFun OutputAllocator);
+		using FTransformMapFunc = void (*)(const FAttributeMap* MapA, const FAttributeMap* MapB, float WeightB, FAttributeMap*& OutputMap, FReallocFun OutputAllocator);
 
 		// Interpolate attribute values
-		// Output = lerp(A, B, Alpha) -> Output = (A * (1.0 - Alpha)) + (B * Alpha)
+		// Output = lerp(A, B, WeightB)
+		// Output = (A * (1.0 - WeightB)) + (B * WeightB)
 		// This function supports outputting to either input
-		static UE_API void Apply(const FAttributeValueTransformerMapPtr& TransformerMap, const FAttributeCollection& CollectionA, const FAttributeCollection& CollectionB, float Alpha, FAttributeCollection& OutputCollection);
+		static UE_API void Apply(const FAttributeValueTransformerMapPtr& TransformerMap, const FAttributeCollection& InputA, const FAttributeCollection& InputB, float WeightB, FAttributeCollection& Output);
+
+		// Interpolate attribute values
+		// If a weight isn't specified, the default weight will be used
+		// Output = lerp(A, B, HasValueWeightB ? ValueWeightB : DefaultWeightB)
+		// This function supports outputting to either input
+		static UE_API void Apply(const FAttributeValueTransformerMapPtr& TransformerMap, const FAttributeCollection& InputA, const FAttributeCollection& InputB, const FAttributeCollection& PerValueWeightsB, float DefaultWeightB, FAttributeCollection& Output);
 	};
 
 	//////////////////////////////////////////////////////////////////////////
@@ -37,8 +47,8 @@ namespace UE::UAF::Transformers
 			using TransformerType = FInterpolate;
 			using ValueType = FBoneTransformAnimationAttribute;
 
-			static void TransformSetMap(const FAttributeSetMap* SetMapA, const FAttributeSetMap* SetMapB, float Alpha, FAttributeSetMap*& OutputSetMap, FReallocFun OutputAllocator);
-			static void TransformMap(const FAttributeMap* MapA, const FAttributeMap* MapB, float Alpha, FAttributeMap*& OutputMap, FReallocFun OutputAllocator);
+			static void TransformSetMap(const FAttributeSetMap* SetMapA, const FAttributeSetMap* SetMapB, const FAttributeSetMap* PerValueWeightMapB, float DefaultWeightB, FAttributeSetMap*& OutputSetMap, FReallocFun OutputAllocator);
+			static void TransformMap(const FAttributeMap* MapA, const FAttributeMap* MapB, float WeightB, FAttributeMap*& OutputMap, FReallocFun OutputAllocator);
 		};
 
 		struct FInterpolate_FloatAttribute
@@ -46,8 +56,8 @@ namespace UE::UAF::Transformers
 			using TransformerType = FInterpolate;
 			using ValueType = FFloatAnimationAttribute;
 
-			static void TransformSetMap(const FAttributeSetMap* SetMapA, const FAttributeSetMap* SetMapB, float Alpha, FAttributeSetMap*& OutputSetMap, FReallocFun OutputAllocator);
-			static void TransformMap(const FAttributeMap* MapA, const FAttributeMap* MapB, float Alpha, FAttributeMap*& OutputMap, FReallocFun OutputAllocator);
+			static void TransformSetMap(const FAttributeSetMap* SetMapA, const FAttributeSetMap* SetMapB, const FAttributeSetMap* PerValueWeightMapB, float DefaultWeightB, FAttributeSetMap*& OutputSetMap, FReallocFun OutputAllocator);
+			static void TransformMap(const FAttributeMap* MapA, const FAttributeMap* MapB, float WeightB, FAttributeMap*& OutputMap, FReallocFun OutputAllocator);
 		};
 	}
 }
diff --git a/Engine/Plugins/Experimental/UAF/UAF/Source/UAF/Public/UAF/Attributes/Transformers/Layer.h b/Engine/Plugins/Experimental/UAF/UAF/Source/UAF/Public/UAF/Attributes/Transformers/Layer.h
new file mode 100644
--- /dev/null
+++ b/Engine/Plugins/Experimental/UAF/UAF/Source/UAF/Public/UAF/Attributes/Transformers/Layer.h
@@ -0,0 +1,76 @@
+// Copyright Epic Games, Inc. All Rights Reserved.
+
+#pragma once
+
+#include "CoreMinimal.h"
+
+#include "UAF/Attributes/EngineAttributes.h"
+#include "UAF/Attributes/Transformers/Transformer.h"
+
+#define UE_API UAF_API
+
+namespace UE::UAF::Transformers
+{
+	// Attribute value transformer to layer values
+	struct FLayer final : public FAttributeValueTransformer
+	{
+		static UE_API FName TransformerName;
+
+		// If the output buffer is nullptr, an allocator is provided and the implementation is responsible for allocating the output buffer and returning it
+		// Note that the input/output buffers can be the same buffer
+		// Per-value map is optional, when not present, the default weight should be used
+		// Unbound values do not support per-value weights
+		using FTransformSetMapFunc = void (*)(const FAttributeSetMap* Base, const FAttributeSetMap* Layer, const FAttributeSetMap* PerValueWeights, float DefaultLayerWeight, FAttributeSetMap*& OutputSetMap, FReallocFun OutputAllocator);
+		using FTransformMapFunc = void (*)(const FAttributeMap* Base, const FAttributeMap* Layer, float LayerWeight, FAttributeMap*& OutputMap, FReallocFun OutputAllocator);
+
+		// Layer values onto a base
+		// Output = layer(Base, Layer, LayerWeight)
+		// Output = ValueInLayer ? lerp(BaseValue, LayerValue, LayerWeight) : BaseValue
+		// Layer values do not have to be present in the Base and extra values are dropped
+		// This function supports outputting to either input
+		static UE_API void Apply(
+			const FAttributeValueTransformerMapPtr& TransformerMap,
+			const FAttributeCollection& Base,
+			const FAttributeCollection& Layer,
+			float LayerWeight,
+			FAttributeCollection& Output);
+
+		// Layer values onto a base
+		// Output = ValueInLayer ? lerp(BaseValue, LayerValue, HasLayerValueWeight ? LayerValueWeight : DefaultLayerWeight) : BaseValue
+		// Layer values do not have to be present in the Base and extra values are dropped
+		// This function supports outputting to either input
+		static UE_API void Apply(
+			const FAttributeValueTransformerMapPtr& TransformerMap,
+			const FAttributeCollection& Base,
+			const FAttributeCollection& Layer,
+			const FAttributeCollection& PerValueLayerWeights,
+			float DefaultLayerWeight,
+			FAttributeCollection& Output);
+	};
+
+	//////////////////////////////////////////////////////////////////////////
+	// Engine built-in value type specializations
+
+	namespace Private
+	{
+		struct FLayer_BoneTransformAttribute
+		{
+			using TransformerType = FLayer;
+			using ValueType = FBoneTransformAnimationAttribute;
+
+			static void TransformSetMap(const FAttributeSetMap* Base, const FAttributeSetMap* Layer, const FAttributeSetMap* PerValueWeights, float DefaultLayerWeight, FAttributeSetMap*& OutputSetMap, FReallocFun OutputAllocator);
+			static void TransformMap(const FAttributeMap* Base, const FAttributeMap* Layer, float LayerWeight, FAttributeMap*& OutputMap, FReallocFun OutputAllocator);
+		};
+
+		struct FLayer_FloatAttribute
+		{
+			using TransformerType = FLayer;
+			using ValueType = FFloatAnimationAttribute;
+
+			static void TransformSetMap(const FAttributeSetMap* Base, const FAttributeSetMap* Layer, const FAttributeSetMap* PerValueWeights, float DefaultLayerWeight, FAttributeSetMap*& OutputSetMap, FReallocFun OutputAllocator);
+			static void TransformMap(const FAttributeMap* Base, const FAttributeMap* Layer, float LayerWeight, FAttributeMap*& OutputMap, FReallocFun OutputAllocator);
+		};
+	}
+}
+
+#undef UE_API
diff --git a/Engine/Plugins/Experimental/UAF/UAFAnimGraph/Source/UAFAnimGraph/Private/EvaluationVM/Tasks/BlendKeyframesPerBone.cpp b/Engine/Plugins/Experimental/UAF/UAFAnimGraph/Source/UAFAnimGraph/Private/EvaluationVM/Tasks/BlendKeyframesPerBone.cpp
--- a/Engine/Plugins/Experimental/UAF/UAFAnimGraph/Source/UAFAnimGraph/Private/EvaluationVM/Tasks/BlendKeyframesPerBone.cpp
+++ b/Engine/Plugins/Experimental/UAF/UAFAnimGraph/Source/UAFAnimGraph/Private/EvaluationVM/Tasks/BlendKeyframesPerBone.cpp
@@ -12,7 +12,9 @@
 #include "EvaluationVM/KeyframeState.h"
 #include "TransformArrayOperations.h"
 #include "UAF/Attributes/Transformers/Accumulate.h"
+#include "UAF/Attributes/Transformers/Layer.h"
 #include "UAF/Attributes/Transformers/Overwrite.h"
+#include "UAF/Attributes/Transformers/Sanitize.h"
 
 #include UE_INLINE_GENERATED_CPP_BY_NAME(BlendKeyframesPerBone)
 
@@ -110,6 +112,113 @@ namespace UE::UAF::Private
 
 		return PerValueWeights;
 	}
+
+	static FAttributeCollectionStack BuildPerValueWeights(
+		const FAttributeCollection& Input,
+		UE::UAF::FEvaluationVM& VM,
+		const USkeleton* SourceSkeleton,
+		const TArray<float>& BoneMaskWeights,
+		const IInterpolationIndexProvider& BlendProfile,
+		const UE::Anim::TNamedValueArray<FDefaultAllocator, UE::Anim::FCurveElement>& CurveMaskWeights,
+		const TArray<FAnimNextBlendKeyframePerBoneWithScaleTask::FMaskedAttributeWeight>& AttributeMaskWeights,
+		float DefaultWeight)
+	{
+		UScriptStruct* FloatAttributeType = FFloatAnimationAttribute::StaticStruct();
+
+		// Build our per-value weights (only for bound values)
+		// By default, the layered values are ignored unless they have weight
+		FAttributeCollectionStack PerValueWeights(VM.GetActiveNamedSet());
+
+		for (auto It = Input.CreateConstSetMapIterator(); It; ++It)
+		{
+			UScriptStruct* AttributeType = It.GetAttributeType();
+
+			// Every attribute type maps to a float, our weight
+			FAttributeMappingKey MappingKey = FAttributeMappingKey::MakeFromTo(AttributeType, FloatAttributeType);
+
+			TAttributeSetMap<FFloatAnimationAttribute>* SetWeights = PerValueWeights.AddSetMap<FFloatAnimationAttribute>(MappingKey);
+
+			if (const TAttributeSetMap<FBoneTransformAnimationAttribute>* BoneMap = Cast<FBoneTransformAnimationAttribute>(It.GetMap()))
+			{
+				const FAttributeTypedSetPtr& BoneSet = BoneMap->GetTypedSet();
+				const int32 NumBonesInSet = BoneMap->Num();
+
+				const USkeleton* TargetSkeleton = VM.GetActiveEvaluationContext().GetSkeleton();
+				TSharedPtr<IInterpolationIndexProvider::FPerBoneInterpolationData> Data = BlendProfile.GetPerBoneInterpolationData(TargetSkeleton);
+
+				const FSkeletonRemapping& SkeletonRemapping = UE::Anim::FSkeletonRemappingRegistry::Get().GetRemapping(SourceSkeleton, TargetSkeleton);
+
+				if (SkeletonRemapping.IsValid())
+				{
+					for (FAttributeSetIndex SetIndex(0); SetIndex < NumBonesInSet; ++SetIndex)
+					{
+						const FAttributeBindingIndex BindingIndex = BoneSet->GetBindingIndex(SetIndex);
+
+						const FSkeletonPoseBoneIndex TargetSkeletonBoneIndex = BindingIndex;
+						const FSkeletonPoseBoneIndex SourceSkeletonBoneIndex(SkeletonRemapping.GetSourceSkeletonBoneIndex(TargetSkeletonBoneIndex.GetInt()));
+
+						const int32 WeightIndex = SourceSkeletonBoneIndex.IsValid() ? BlendProfile.GetPerBoneInterpolationIndex(SourceSkeletonBoneIndex, TargetSkeleton, Data.Get()) : INDEX_NONE;
+						(*SetWeights)[SetIndex].Value = WeightIndex != INDEX_NONE ? BoneMaskWeights[WeightIndex] : 0.0f;
+						(*SetWeights)[SetIndex].Value *= DefaultWeight;
+					}
+				}
+				else
+				{
+					for (FAttributeSetIndex SetIndex(0); SetIndex < NumBonesInSet; ++SetIndex)
+					{
+						const FAttributeBindingIndex BindingIndex = BoneSet->GetBindingIndex(SetIndex);
+
+						const FSkeletonPoseBoneIndex TargetSkeletonBoneIndex = BindingIndex;
+
+						const int32 WeightIndex = BlendProfile.GetPerBoneInterpolationIndex(TargetSkeletonBoneIndex, TargetSkeleton, Data.Get());
+						(*SetWeights)[SetIndex].Value = WeightIndex != INDEX_NONE ? BoneMaskWeights[WeightIndex] : 0.0f;
+						(*SetWeights)[SetIndex].Value *= DefaultWeight;
+					}
+				}
+			}
+			else if (const TAttributeSetMap<FFloatAnimationAttribute>* FloatMap = Cast<FFloatAnimationAttribute>(It.GetMap()))
+			{
+				SetWeights->FillWithValue(FFloatAnimationAttribute{ 0.0f });
+
+				const FAttributeTypedSetPtr& FloatSet = FloatMap->GetTypedSet();
+				CurveMaskWeights.ForEachElement([&FloatSet, SetWeights, DefaultWeight](const UE::Anim::FCurveElement& Curve)
+					{
+						const FAttributeSetIndex SetIndex = FloatSet->FindIndex(Curve.Name);
+						if (SetIndex.IsValid())
+						{
+							(*SetWeights)[SetIndex].Value = Curve.Value;
+							(*SetWeights)[SetIndex].Value *= DefaultWeight;
+						}
+					});
+
+				for (const FAnimNextBlendKeyframePerBoneWithScaleTask::FMaskedAttributeWeight& AttributeWeight : AttributeMaskWeights)
+				{
+					const FAttributeSetIndex SetIndex = FloatSet->FindIndex(AttributeWeight.Attribute.GetName());
+					if (SetIndex.IsValid())
+					{
+						(*SetWeights)[SetIndex].Value = AttributeWeight.Weight;
+						(*SetWeights)[SetIndex].Value *= DefaultWeight;
+					}
+				}
+			}
+			else
+			{
+				SetWeights->FillWithValue(FFloatAnimationAttribute{ 0.0f });
+
+				const FAttributeTypedSetPtr& TypedSet = It.GetMap()->GetTypedSet();
+				for (const FAnimNextBlendKeyframePerBoneWithScaleTask::FMaskedAttributeWeight& AttributeWeight : AttributeMaskWeights)
+				{
+					const FAttributeSetIndex SetIndex = TypedSet->FindIndex(AttributeWeight.Attribute.GetName());
+					if (SetIndex.IsValid())
+					{
+						(*SetWeights)[SetIndex].Value = AttributeWeight.Weight;
+					}
+				}
+			}
+		}
+
+		return PerValueWeights;
+	}
 }
 
 FAnimNextBlendOverwriteKeyframePerBoneWithScaleTask FAnimNextBlendOverwriteKeyframePerBoneWithScaleTask::Make(const IBlendProfileInterface* BlendProfile, const FBlendSampleData& BlendData, float ScaleFactor)
@@ -433,13 +542,6 @@ void FAnimNextBlendKeyframePerBoneWithScaleTask::Execute(UE::UAF::FEvaluationVM&
 {
 	using namespace UE::UAF;
 
-	if (VM.GetActiveNamedSet())
-	{
-		// TODO: Implement with new attribute runtime
-		Super::Execute(VM);
-		return;
-	}
-
 	if (BlendProfile == nullptr)
 	{
 		// If we don't have a blend profile, we blend the whole pose
@@ -447,222 +549,250 @@ void FAnimNextBlendKeyframePerBoneWithScaleTask::Execute(UE::UAF::FEvaluationVM&
 		return;
 	}
 
-	// Pop our top two poses, we'll re-use the top keyframe for our result
-
-	TUniquePtr<FKeyframeState> KeyframeB;
-	if (!VM.PopValue(KEYFRAME_STACK_NAME, KeyframeB))
+	if (VM.GetActiveNamedSet())
 	{
-		// We have no inputs, nothing to do
-		return;
-	}
+		TUniquePtr<FAttributeCollection> Layer;
+		if (!VM.PopValue(ATTRIBUTE_STACK_NAME, Layer))
+		{
+			// We have no inputs, nothing to do
+			return;
+		}
 
-	TUniquePtr<FKeyframeState> KeyframeA;
-	if (!VM.PopValue(KEYFRAME_STACK_NAME, KeyframeA))
-	{
-		// We have a single input, leave it on top of the stack
-		VM.PushValue(KEYFRAME_STACK_NAME, MoveTemp(KeyframeB));
-		return;
+		TUniquePtr<FAttributeCollection> Base;
+		if (!VM.PopValue(ATTRIBUTE_STACK_NAME, Base))
+		{
+			// We have a single input, leave it on top of the stack
+			VM.PushValue(ATTRIBUTE_STACK_NAME, MoveTemp(Layer));
+			return;
+		}
+
+		const FAttributeCollectionStack PerValueWeightsLayer = Private::BuildPerValueWeights(*Base, VM, SourceSkeleton, BoneMaskWeights, *BlendProfile, CurveMaskWeights, AttributeMaskWeights, ScaleFactor);
+
+		// Apply our layer onto our base and modify it in-place
+		Transformers::FLayer::Apply(VM.GetTransformerMap(), *Base, *Layer, PerValueWeightsLayer, ScaleFactor, *Base);
+		Transformers::FSanitize::Apply(VM.GetTransformerMap(), *Base);
+
+		VM.PushValue(ATTRIBUTE_STACK_NAME, MoveTemp(Base));
 	}
+	else
+	{
+		// Pop our top two poses, we'll re-use the top keyframe for our result
 
-	const USkeleton* TargetSkeleton = KeyframeB->Pose.GetSkeletonAsset();
-	TSharedPtr<IInterpolationIndexProvider::FPerBoneInterpolationData> Data = BlendProfile->GetPerBoneInterpolationData(TargetSkeleton);
+		TUniquePtr<FKeyframeState> KeyframeB;
+		if (!VM.PopValue(KEYFRAME_STACK_NAME, KeyframeB))
+		{
+			// We have no inputs, nothing to do
+			return;
+		}
 
-	const FSkeletonRemapping& SkeletonRemapping = UE::Anim::FSkeletonRemappingRegistry::Get().GetRemapping(SourceSkeleton, TargetSkeleton);
+		TUniquePtr<FKeyframeState> KeyframeA;
+		if (!VM.PopValue(KEYFRAME_STACK_NAME, KeyframeA))
+		{
+			// We have a single input, leave it on top of the stack
+			VM.PushValue(KEYFRAME_STACK_NAME, MoveTemp(KeyframeB));
+			return;
+		}
 
-	const TArrayView<const FBoneIndexType> LODBoneIndexToSkeletonBoneIndexMap = KeyframeB->Pose.GetLODBoneIndexToSkeletonBoneIndexMap();
-	const int32 NumLODBones = LODBoneIndexToSkeletonBoneIndexMap.Num();
+		const USkeleton* TargetSkeleton = KeyframeB->Pose.GetSkeletonAsset();
+		TSharedPtr<IInterpolationIndexProvider::FPerBoneInterpolationData> Data = BlendProfile->GetPerBoneInterpolationData(TargetSkeleton);
 
-	TArray<int32, TMemStackAllocator<>> LODBoneIndexToWeightIndexMap;
-	LODBoneIndexToWeightIndexMap.AddUninitialized(NumLODBones);
+		const FSkeletonRemapping& SkeletonRemapping = UE::Anim::FSkeletonRemappingRegistry::Get().GetRemapping(SourceSkeleton, TargetSkeleton);
 
-	if (SkeletonRemapping.IsValid())
-	{
-		for (int32 LODBoneIndex = 0; LODBoneIndex < NumLODBones; ++LODBoneIndex)
+		const TArrayView<const FBoneIndexType> LODBoneIndexToSkeletonBoneIndexMap = KeyframeB->Pose.GetLODBoneIndexToSkeletonBoneIndexMap();
+		const int32 NumLODBones = LODBoneIndexToSkeletonBoneIndexMap.Num();
+
+		TArray<int32, TMemStackAllocator<>> LODBoneIndexToWeightIndexMap;
+		LODBoneIndexToWeightIndexMap.AddUninitialized(NumLODBones);
+
+		if (SkeletonRemapping.IsValid())
 		{
-			const FSkeletonPoseBoneIndex TargetSkeletonBoneIndex(LODBoneIndexToSkeletonBoneIndexMap[LODBoneIndex]);
-			const FSkeletonPoseBoneIndex SourceSkeletonBoneIndex(SkeletonRemapping.GetSourceSkeletonBoneIndex(TargetSkeletonBoneIndex.GetInt()));
-			LODBoneIndexToWeightIndexMap[LODBoneIndex] = SourceSkeletonBoneIndex.IsValid() ? BlendProfile->GetPerBoneInterpolationIndex(SourceSkeletonBoneIndex, TargetSkeleton, Data.Get()) : INDEX_NONE;
+			for (int32 LODBoneIndex = 0; LODBoneIndex < NumLODBones; ++LODBoneIndex)
+			{
+				const FSkeletonPoseBoneIndex TargetSkeletonBoneIndex(LODBoneIndexToSkeletonBoneIndexMap[LODBoneIndex]);
+				const FSkeletonPoseBoneIndex SourceSkeletonBoneIndex(SkeletonRemapping.GetSourceSkeletonBoneIndex(TargetSkeletonBoneIndex.GetInt()));
+				LODBoneIndexToWeightIndexMap[LODBoneIndex] = SourceSkeletonBoneIndex.IsValid() ? BlendProfile->GetPerBoneInterpolationIndex(SourceSkeletonBoneIndex, TargetSkeleton, Data.Get()) : INDEX_NONE;
+			}
 		}
-	}
-	else
-	{
-		for (int32 LODBoneIndex = 0; LODBoneIndex < NumLODBones; ++LODBoneIndex)
+		else
 		{
-			const FSkeletonPoseBoneIndex SkeletonBoneIndex(LODBoneIndexToSkeletonBoneIndexMap[LODBoneIndex]);
-			LODBoneIndexToWeightIndexMap[LODBoneIndex] = BlendProfile->GetPerBoneInterpolationIndex(SkeletonBoneIndex, TargetSkeleton, Data.Get());
+			for (int32 LODBoneIndex = 0; LODBoneIndex < NumLODBones; ++LODBoneIndex)
+			{
+				const FSkeletonPoseBoneIndex SkeletonBoneIndex(LODBoneIndexToSkeletonBoneIndexMap[LODBoneIndex]);
+				LODBoneIndexToWeightIndexMap[LODBoneIndex] = BlendProfile->GetPerBoneInterpolationIndex(SkeletonBoneIndex, TargetSkeleton, Data.Get());
+			}
 		}
-	}
 
-	if (EnumHasAnyFlags(VM.GetFlags(), EEvaluationFlags::Bones))
-	{
-		check(KeyframeA->Pose.GetNumBones() == KeyframeB->Pose.GetNumBones());
-
-		BlendOverwritePerBoneWithScale(
-			KeyframeA->Pose.LocalTransforms.GetView(), KeyframeA->Pose.LocalTransforms.GetConstView(),
-			LODBoneIndexToWeightIndexMap, BoneMaskWeights, 1.0f, true);
+		if (EnumHasAnyFlags(VM.GetFlags(), EEvaluationFlags::Bones))
+		{
+			check(KeyframeA->Pose.GetNumBones() == KeyframeB->Pose.GetNumBones());
 
-		BlendAddPerBoneWithScale(
-			KeyframeA->Pose.LocalTransforms.GetView(), KeyframeB->Pose.LocalTransforms.GetConstView(),
-			LODBoneIndexToWeightIndexMap, BoneMaskWeights, ScaleFactor);
-	}
+			BlendOverwritePerBoneWithScale(
+				KeyframeA->Pose.LocalTransforms.GetView(), KeyframeA->Pose.LocalTransforms.GetConstView(),
+				LODBoneIndexToWeightIndexMap, BoneMaskWeights, 1.0f, true);
 
-	if (EnumHasAnyFlags(VM.GetFlags(), EEvaluationFlags::Curves))
-	{
-		FBlendedCurve& OutCurve = KeyframeA->Curves;
-		FBlendedCurve& TargetCurve = KeyframeB->Curves;
+			BlendAddPerBoneWithScale(
+				KeyframeA->Pose.LocalTransforms.GetView(), KeyframeB->Pose.LocalTransforms.GetConstView(),
+				LODBoneIndexToWeightIndexMap, BoneMaskWeights, ScaleFactor);
+		}
 
-		if (FAnimWeight::IsRelevant(ScaleFactor))
+		if (EnumHasAnyFlags(VM.GetFlags(), EEvaluationFlags::Curves))
 		{
-			FBlendedCurve FilteredCurves;
+			FBlendedCurve& OutCurve = KeyframeA->Curves;
+			FBlendedCurve& TargetCurve = KeyframeB->Curves;
 
-			// Multiply per-curve blend weights by matching blend pose curves
-			UE::Anim::FNamedValueArrayUtils::Intersection(
-				TargetCurve,
-				CurveMaskWeights,
-				[this, &FilteredCurves](const UE::Anim::FCurveElement& InBlendElement, const UE::Anim::FCurveElement& InMaskElement) mutable
-				{
-					FilteredCurves.Add(InBlendElement.Name, InBlendElement.Value * InMaskElement.Value);
-				});
+			if (FAnimWeight::IsRelevant(ScaleFactor))
+			{
+				FBlendedCurve FilteredCurves;
 
-			// Override blend curve values with premultipled curves
-			TargetCurve.Combine(FilteredCurves);
+				// Multiply per-curve blend weights by matching blend pose curves
+				UE::Anim::FNamedValueArrayUtils::Intersection(
+					TargetCurve,
+					CurveMaskWeights,
+					[this, &FilteredCurves](const UE::Anim::FCurveElement& InBlendElement, const UE::Anim::FCurveElement& InMaskElement) mutable
+					{
+						FilteredCurves.Add(InBlendElement.Name, InBlendElement.Value * InMaskElement.Value);
+					});
 
-			// Remove curves that have been filtered by the mask, curves with no mask value defined remain, even with 0.0 value
-			UE::Anim::FNamedValueArrayUtils::RemoveByPredicate(
-				TargetCurve,
-				CurveMaskWeights,
-				[](const UE::Anim::FCurveElement& InBaseElement, const UE::Anim::FCurveElement& InMaskElement)
-				{
-					const bool bKeep = InMaskElement.Value == 0.0;
-					return bKeep;
-				}
-			);
+				// Override blend curve values with premultipled curves
+				TargetCurve.Combine(FilteredCurves);
 
-			// Combine base and filtered pre-multiplied blend curves
-			UE::Anim::FNamedValueArrayUtils::Union(
-				OutCurve,
-				TargetCurve,
-				[this](UE::Anim::FCurveElement& InOutBaseElement, const UE::Anim::FCurveElement& InBlendElement, UE::Anim::ENamedValueUnionFlags InFlags)
-				{
-					if (InFlags == UE::Anim::ENamedValueUnionFlags::BothArgsValid || InFlags == UE::Anim::ENamedValueUnionFlags::ValidArg1)
+				// Remove curves that have been filtered by the mask, curves with no mask value defined remain, even with 0.0 value
+				UE::Anim::FNamedValueArrayUtils::RemoveByPredicate(
+					TargetCurve,
+					CurveMaskWeights,
+					[](const UE::Anim::FCurveElement& InBaseElement, const UE::Anim::FCurveElement& InMaskElement)
 					{
-						InOutBaseElement.Value = FMath::Lerp(InOutBaseElement.Value, InBlendElement.Value, ScaleFactor);
-						InOutBaseElement.Flags |= InBlendElement.Flags;
+						const bool bKeep = InMaskElement.Value == 0.0;
+						return bKeep;
 					}
-				});
-		}
-	}
+				);
 
-	if (EnumHasAnyFlags(VM.GetFlags(), EEvaluationFlags::Attributes))
-	{
-		using namespace UE::Anim;
+				// Combine base and filtered pre-multiplied blend curves
+				UE::Anim::FNamedValueArrayUtils::Union(
+					OutCurve,
+					TargetCurve,
+					[this](UE::Anim::FCurveElement& InOutBaseElement, const UE::Anim::FCurveElement& InBlendElement, UE::Anim::ENamedValueUnionFlags InFlags)
+					{
+						if (InFlags == UE::Anim::ENamedValueUnionFlags::BothArgsValid || InFlags == UE::Anim::ENamedValueUnionFlags::ValidArg1)
+						{
+							InOutBaseElement.Value = FMath::Lerp(InOutBaseElement.Value, InBlendElement.Value, ScaleFactor);
+							InOutBaseElement.Flags |= InBlendElement.Flags;
+						}
+					});
+			}
+		}
 
-		FStackAttributeContainer& BaseAttributes = KeyframeA->Attributes;
-		FStackAttributeContainer& BlendAttributes = KeyframeB->Attributes;
+		if (EnumHasAnyFlags(VM.GetFlags(), EEvaluationFlags::Attributes))
+		{
+			using namespace UE::Anim;
 
-		FStackAttributeContainer OutputAttributes;
+			FStackAttributeContainer& BaseAttributes = KeyframeA->Attributes;
+			FStackAttributeContainer& BlendAttributes = KeyframeB->Attributes;
 
-		// Attributes are to be masked out according to the mask weights in AttributeMaskWeights, if an attribute has no mask weight set then it
-		// inherits the weight of whatever bone it is attached to. Below are possible configurations that we need to account for:
+			FStackAttributeContainer OutputAttributes;
 
-		// Root 0.0						Root set to 0.0 therefore a RootMotionDelta attribute will also be masked out without having to set an explicit entry in AttributeMaskWeights
+			// Attributes are to be masked out according to the mask weights in AttributeMaskWeights, if an attribute has no mask weight set then it
+			// inherits the weight of whatever bone it is attached to. Below are possible configurations that we need to account for:
 
-		// Root 0.0						RootMotionDelta is set to 1.0 in AttributeMaskWeights despite the parent bone being masked out.
-		//  \ RootMotionDelta 1.0
+			// Root 0.0						Root set to 0.0 therefore a RootMotionDelta attribute will also be masked out without having to set an explicit entry in AttributeMaskWeights
 
-		// Root 1.0						RootMotionDelta is being masked out in AttributeMaskWeights despite the parent bone being kept.
-		//  \ RootMotionDelta 0.0
+			// Root 0.0						RootMotionDelta is set to 1.0 in AttributeMaskWeights despite the parent bone being masked out.
+			//  \ RootMotionDelta 1.0
 
+			// Root 1.0						RootMotionDelta is being masked out in AttributeMaskWeights despite the parent bone being kept.
+			//  \ RootMotionDelta 0.0
 
-		// Below is a table of the possible permutations of base/blend attributes being present/absent along with the possible mask values
-		// k denotes some value in the range (0, 1) exclusive.
-		// - denotes an absent attribute
 
-		// Base | Blend | Weight | Output
-		// ------------------------------
-		// a    | b     | 1.0    | b
-		// a    | b     | k      | lerp(a, b, k)
-		// a    | b     | 0.0    | a
-		// - - - - - - - - - - - - - - -
-		// a    | -     | 1.0    | a
-		// a    | -     | k      | a
-		// a    | -     | 0.0    | a
-		// - - - - - - - - - - - - - - -
-		// -    | b     | 1.0    | b
-		// -    | b     | k      | lerp(default, b, k)
-		// -    | b     | 0.0    | -
+			// Below is a table of the possible permutations of base/blend attributes being present/absent along with the possible mask values
+			// k denotes some value in the range (0, 1) exclusive.
+			// - denotes an absent attribute
 
+			// Base | Blend | Weight | Output
+			// ------------------------------
+			// a    | b     | 1.0    | b
+			// a    | b     | k      | lerp(a, b, k)
+			// a    | b     | 0.0    | a
+			// - - - - - - - - - - - - - - -
+			// a    | -     | 1.0    | a
+			// a    | -     | k      | a
+			// a    | -     | 0.0    | a
+			// - - - - - - - - - - - - - - -
+			// -    | b     | 1.0    | b
+			// -    | b     | k      | lerp(default, b, k)
+			// -    | b     | 0.0    | -
 
-		// 1. Blend attributes according to the bone blend weights,
-		// i.e. an attribute's weight is determined by the weight of its attached bone.
-		UE::Anim::Attributes::BlendAttributesPerBone(BaseAttributes, BlendAttributes, BoneMaskWeights, OutputAttributes);
 
+			// 1. Blend attributes according to the bone blend weights,
+			// i.e. an attribute's weight is determined by the weight of its attached bone.
+			UE::Anim::Attributes::BlendAttributesPerBone(BaseAttributes, BlendAttributes, BoneMaskWeights, OutputAttributes);
 
-		// 2. For each attribute that has a custom weight, i.e. one's that shouldn't be weighted by its attached bone, go and correct the blended value.
-		for (const FMaskedAttributeWeight& MaskedAttribute : AttributeMaskWeights)
-		{
-			const TConstArrayView<TWeakObjectPtr<UScriptStruct>> UniqueTypes = OutputAttributes.GetUniqueTypes();
 
-			for (int32 UniqueTypeIndex = 0; UniqueTypeIndex < UniqueTypes.Num(); ++UniqueTypeIndex)
+			// 2. For each attribute that has a custom weight, i.e. one's that shouldn't be weighted by its attached bone, go and correct the blended value.
+			for (const FMaskedAttributeWeight& MaskedAttribute : AttributeMaskWeights)
 			{
-				const TConstArrayView<FAttributeId> AttributeKeys = OutputAttributes.GetKeys(UniqueTypeIndex);
-				const TConstArrayView<TWrappedAttribute<FAnimStackAllocator>> AttributeValues = OutputAttributes.GetValues(UniqueTypeIndex);
-				
-				const TWeakObjectPtr<UScriptStruct> AttributeType = UniqueTypes[UniqueTypeIndex];
-				const IAttributeBlendOperator* Operator = UE::Anim::AttributeTypes::GetTypeOperator(AttributeType);
-
-				TWrappedAttribute<FAnimStackAllocator> DefaultData(AttributeType.Get());
-				AttributeType->InitializeStruct(DefaultData.GetPtr<void>());
+				const TConstArrayView<TWeakObjectPtr<UScriptStruct>> UniqueTypes = OutputAttributes.GetUniqueTypes();
 
-				uint8* OutputData = OutputAttributes.Find(AttributeType.Get(), MaskedAttribute.Attribute);
-				if (OutputData)
+				for (int32 UniqueTypeIndex = 0; UniqueTypeIndex < UniqueTypes.Num(); ++UniqueTypeIndex)
 				{
-					uint8* BaseData = BaseAttributes.Find(AttributeType.Get(), MaskedAttribute.Attribute);
-					uint8* BlendData = BlendAttributes.Find(AttributeType.Get(), MaskedAttribute.Attribute);
+					const TConstArrayView<FAttributeId> AttributeKeys = OutputAttributes.GetKeys(UniqueTypeIndex);
+					const TConstArrayView<TWrappedAttribute<FAnimStackAllocator>> AttributeValues = OutputAttributes.GetValues(UniqueTypeIndex);
 
-					if (BaseData && BlendData)
-					{
-						// Base | Blend | Weight | Output
-						// ------------------------------
-						// a    | b     | 1.0    | b
-						// a    | b     | k      | lerp(a, b, k)
-						// a    | b     | 0.0    | a
+					const TWeakObjectPtr<UScriptStruct> AttributeType = UniqueTypes[UniqueTypeIndex];
+					const IAttributeBlendOperator* Operator = UE::Anim::AttributeTypes::GetTypeOperator(AttributeType);
 
-						Operator->Interpolate(BaseData, BlendData, MaskedAttribute.Weight, OutputData);
-					}
-					else if (BaseData && !BlendData)
-					{
-						// Base | Blend | Weight | Output
-						// ------------------------------
-						// a    | -     | 1.0    | a
-						// a    | -     | k      | a
-						// a    | -     | 0.0    | a
+					TWrappedAttribute<FAnimStackAllocator> DefaultData(AttributeType.Get());
+					AttributeType->InitializeStruct(DefaultData.GetPtr<void>());
 
-						*OutputData = *BaseData;
-					}
-					else if (!BaseData && BlendData)
+					uint8* OutputData = OutputAttributes.Find(AttributeType.Get(), MaskedAttribute.Attribute);
+					if (OutputData)
 					{
-						// Base | Blend | Weight | Output
-						// ------------------------------
-						// -    | b     | 1.0    | b
-						// -    | b     | k      | lerp(default, b, k)
-						// -    | b     | 0.0    | -
+						uint8* BaseData = BaseAttributes.Find(AttributeType.Get(), MaskedAttribute.Attribute);
+						uint8* BlendData = BlendAttributes.Find(AttributeType.Get(), MaskedAttribute.Attribute);
 
-						if (MaskedAttribute.Weight != 0.0f)
+						if (BaseData && BlendData)
 						{
-							Operator->Interpolate(DefaultData.GetPtr<void>(), BlendData, MaskedAttribute.Weight, OutputData);
+							// Base | Blend | Weight | Output
+							// ------------------------------
+							// a    | b     | 1.0    | b
+							// a    | b     | k      | lerp(a, b, k)
+							// a    | b     | 0.0    | a
+
+							Operator->Interpolate(BaseData, BlendData, MaskedAttribute.Weight, OutputData);
 						}
-						else
+						else if (BaseData && !BlendData)
+						{
+							// Base | Blend | Weight | Output
+							// ------------------------------
+							// a    | -     | 1.0    | a
+							// a    | -     | k      | a
+							// a    | -     | 0.0    | a
+
+							*OutputData = *BaseData;
+						}
+						else if (!BaseData && BlendData)
 						{
-							OutputAttributes.Remove(AttributeType.Get(), MaskedAttribute.Attribute);
+							// Base | Blend | Weight | Output
+							// ------------------------------
+							// -    | b     | 1.0    | b
+							// -    | b     | k      | lerp(default, b, k)
+							// -    | b     | 0.0    | -
+
+							if (MaskedAttribute.Weight != 0.0f)
+							{
+								Operator->Interpolate(DefaultData.GetPtr<void>(), BlendData, MaskedAttribute.Weight, OutputData);
+							}
+							else
+							{
+								OutputAttributes.Remove(AttributeType.Get(), MaskedAttribute.Attribute);
+							}
 						}
 					}
-				}				
+				}
 			}
+
+			BaseAttributes.MoveFrom(OutputAttributes);
 		}
 
-		BaseAttributes.MoveFrom(OutputAttributes);
+		VM.PushValue(KEYFRAME_STACK_NAME, MoveTemp(KeyframeA));
 	}
-
-	VM.PushValue(KEYFRAME_STACK_NAME, MoveTemp(KeyframeA));
 }
