Improve sync group phase matching in UAF

#Synced-CL 48987520

diff --git a/Engine/Plugins/Experimental/UAF/UAFAnimGraph/Source/UAFAnimGraph/Private/Graph/SyncGroup_GraphInstanceComponent.cpp b/Engine/Plugins/Experimental/UAF/UAFAnimGraph/Source/UAFAnimGraph/Private/Graph/SyncGroup_GraphInstanceComponent.cpp
--- a/Engine/Plugins/Experimental/UAF/UAFAnimGraph/Source/UAFAnimGraph/Private/Graph/SyncGroup_GraphInstanceComponent.cpp
+++ b/Engine/Plugins/Experimental/UAF/UAFAnimGraph/Source/UAFAnimGraph/Private/Graph/SyncGroup_GraphInstanceComponent.cpp
@@ -654,6 +654,11 @@ namespace UE::UAF
 			OutNextMarker.MarkerIndex = NextMarkerIndex;
 		}
 
+		static FName GetMarkerName(int32 MarkerIndex, const TArrayView<const FTimelineSyncMarker>& SyncMarkers)
+		{
+			return SyncMarkers.IsValidIndex(MarkerIndex) ? SyncMarkers[MarkerIndex].GetName() : NAME_None;
+		}
+
 		static FName GetMarkerName(int32 MarkerIndex, const FTimelineSyncMarkerArray& SyncMarkers)
 		{
 			return SyncMarkers.IsValidIndex(MarkerIndex) ? SyncMarkers[MarkerIndex].GetName() : NAME_None;
@@ -846,8 +851,8 @@ namespace UE::UAF
 			float CurrentTime = TimelineState.GetPosition();
 
 			FSyncGroupPhasePosition PhasePosition;
-			float PrevTime = 0.0f;
-			float NextTime = Duration;
+			float PrevTime;
+			float NextTime;
 
 			// Get previous marker's time and name.
 			if (PrevMarker.MarkerIndex != MarkerIndexSpecialValues::AnimationBoundary && ensureAlwaysMsgf(SyncMarkers.IsValidIndex(PrevMarker.MarkerIndex),
@@ -857,6 +862,12 @@ namespace UE::UAF
 				PhasePosition.PrevMarkerIndex = PrevMarker.MarkerIndex;
 				PhasePosition.PrevMarkerName = SyncMarkers[PrevMarker.MarkerIndex].GetName();
 			}
+			else
+			{
+				PrevTime = 0.0f;
+				PhasePosition.PrevMarkerIndex = MarkerIndexSpecialValues::AnimationBoundary;
+				PhasePosition.PrevMarkerName = NAME_None;
+			}
 
 			// Get next marker's time and name.
 			if (NextMarker.MarkerIndex != MarkerIndexSpecialValues::AnimationBoundary && ensureAlwaysMsgf(SyncMarkers.IsValidIndex(NextMarker.MarkerIndex),
@@ -866,6 +877,12 @@ namespace UE::UAF
 				PhasePosition.NextMarkerIndex = NextMarker.MarkerIndex;
 				PhasePosition.NextMarkerName = SyncMarkers[NextMarker.MarkerIndex].GetName();
 			}
+			else
+			{
+				NextTime = Duration;
+				PhasePosition.NextMarkerIndex = MarkerIndexSpecialValues::AnimationBoundary;
+				PhasePosition.NextMarkerName = NAME_None;
+			}
 
 			// Account for looping
 			if (PrevTime > NextTime)
@@ -1197,99 +1214,219 @@ namespace UE::UAF
 			}
 		}
 
+		static int32 GetNumLoopingMarkers(const FTimelineState& TimelineState, const TArrayView<const FTimelineSyncMarker>& SyncMarkers)
+		{
+			int32 NumSyncMarkers = SyncMarkers.Num();
+
+			if (NumSyncMarkers != 0)
+			{
+				// Determine if our first and last markers match and repeat on the first/last keyframes in which case we don't have a phase in between
+
+				const FTimelineSyncMarker& FirstMarker = SyncMarkers[0];
+				const FTimelineSyncMarker& LastMarker = SyncMarkers.Last();
+				if (FirstMarker.GetName() == LastMarker.GetName())
+				{
+					if (FirstMarker.GetPosition() == 0.0f && LastMarker.GetPosition() == TimelineState.GetDuration())
+					{
+						// We perfectly overlap, ignore our last marker
+						NumSyncMarkers--;
+					}
+				}
+			}
+
+			return NumSyncMarkers;
+		}
+
 		// Returns whether or not two members have phases that can be matched
-		static bool CanMatchPhase(const TArrayView<const FTimelineSyncMarker>& SyncMarkersA, const TArrayView<const FTimelineSyncMarker>& SyncMarkersB)
+		static bool CanMatchPhase(
+			const FTimelineState& TimelineStateA, const TArrayView<const FTimelineSyncMarker>& SyncMarkersA,
+			const FTimelineState& TimelineStateB, const TArrayView<const FTimelineSyncMarker>& SyncMarkersB)
 		{
 			// TODO: We can cache this value once we've computed this between a follower and leader
 			// We only need to recompute this if the leader changes or when joining
 
-			const TArrayView<const FTimelineSyncMarker>* LongestSyncMarkersList;
-			const TArrayView<const FTimelineSyncMarker>* ShortestSyncMarkersList;
+			// Identify the number of markers contained by a phase cycle
+			// We only support phase matching if our cycles match
+			// Imagine we have a marker sequence with: RLRL
+			// Our phase cycle is: RL
+			// We only consider simple cycles where the first marker matches the last and isn't present within the cycle
+			// e.g. a cycle cannot contain the first marker, RLR would not be valid for example as part of RLRRLR
+			// 
+			// Once we've identified our cycle, we can determine whether or not two clips can phase match
+			// There are 3 cases to consider:
+			//   - Neither clip loops
+			//   - One clip loops and the other doesn't
+			//   - Both clips loop
+			// 
+			// We also have to consider looping clips where the last keyframe matches the first (we don't interpolate between)
+			// This occurs and is necessary when root motion is used which means the first marker can be repeated: RLRLR
 
-			if (SyncMarkersA.Num() >= SyncMarkersB.Num())
+			if (!TimelineStateA.IsLooping() && !TimelineStateB.IsLooping())
 			{
-				LongestSyncMarkersList = &SyncMarkersA;
-				ShortestSyncMarkersList = &SyncMarkersB;
+				// Neither are looping
+				// Both must have a matching number of cycles(we don't support partial matching, e.g. RLRL and RLRLRL matching
+				// in the first portion since the first would 'freeze' once it reaches the last marker if it doesn't loop)
+				// Here we can also support partial cycles as long as both match (e.g. RLR and RLR)
+
+				const int32 NumSyncMarkersA = SyncMarkersA.Num();
+				const int32 NumSyncMarkersB = SyncMarkersB.Num();
+
+				if (NumSyncMarkersA != NumSyncMarkersB)
+				{
+					return false;	// Cannot match if our marker count differs
+				}
+
+				if (NumSyncMarkersA == 0)
+				{
+					return false;	// Cannot match if we have no markers
+				}
+
+				for (int32 MarkerIdx = 0; MarkerIdx < NumSyncMarkersA; ++MarkerIdx)
+				{
+					if (SyncMarkersA[MarkerIdx].GetName() != SyncMarkersB[MarkerIdx].GetName())
+					{
+						// Markers do not match, we cannot phase match
+						return false;
+					}
+				}
+
+				// All our markers match
+				return true;
 			}
-			else
+			else if (TimelineStateA.IsLooping() && TimelineStateB.IsLooping())
 			{
-				LongestSyncMarkersList = &SyncMarkersB;
-				ShortestSyncMarkersList = &SyncMarkersA;
-			}
+				// Both are looping
+				// They can phase match as long as the marker sequence loops properly over their combined length
+				// If we have RL and RLRL, our combined length is our greatest common divisor: 4 markers
+				// We'd get: RL RL and RLRL
+				// If we have RLRL and RLRLRL, our combined length is our greatest common divisor: 12 markers
+				// We'd get: RLRL RLRL RLRL and RLRLRL RLRLRL
 
-			const int32 NumLongestSyncMarkers = LongestSyncMarkersList->Num();
-			const int32 NumShortestSyncMarkers = ShortestSyncMarkersList->Num();
+				const int32 NumSyncMarkersA = GetNumLoopingMarkers(TimelineStateA, SyncMarkersA);
+				const int32 NumSyncMarkersB = GetNumLoopingMarkers(TimelineStateB, SyncMarkersB);
 
-			if (NumShortestSyncMarkers == 0)
-			{
-				return false;	// Can't match if we have no markers
-			}
+				if (NumSyncMarkersA == 0 || NumSyncMarkersB == 0)
+				{
+					return false;	// Cannot match if we have no markers
+				}
 
-			for (int32 LongestIdx = 0; LongestIdx < NumLongestSyncMarkers; ++LongestIdx)
-			{
-				int32 ShortestIdx;
-				if (SyncMarkersA[0].GetName() == SyncMarkersA.Last().GetName())
+				const int32 NumTargetMarkers = FMath::GreatestCommonDivisor(NumSyncMarkersA, NumSyncMarkersB);
+				for (int32 TargetMarkerIndex = 0; TargetMarkerIndex < NumTargetMarkers; ++TargetMarkerIndex)
 				{
-					// Looping markers
-					ShortestIdx = LongestIdx <= NumShortestSyncMarkers ? LongestIdx : ((LongestIdx % NumShortestSyncMarkers) + 1);
+					const int32 MarkerIndexA = TargetMarkerIndex % NumSyncMarkersA;
+					const int32 MarkerIndexB = TargetMarkerIndex % NumSyncMarkersB;
+
+					if (SyncMarkersA[MarkerIndexA].GetName() != SyncMarkersB[MarkerIndexB].GetName())
+					{
+						// Markers do not match, we cannot phase match
+						return false;
+					}
 				}
-				else
+
+				// Markers are perfectly in sync
+				return true;
+			}
+			else
+			{
+				// One of them is looping
+				// The non-looping clip must have a matching multiple number of cycles meaning the looping
+				// clip must fully loop a number of times (not partially) (e.g. RLRL (looping) and RLRLRL cannot phase match but RLRL (looping) and RLRLRLRL can)
+				// This is because we do not track how many times we've looped and so if the shorter sequence is leading, we do not know how far into the second
+				// sequence we should be which is problematic if the short leading sequence is the one looping
+
+				const FTimelineState& LoopingTimelineState = TimelineStateA.IsLooping() ? TimelineStateA : TimelineStateB;
+				const TArrayView<const FTimelineSyncMarker>& LoopingSyncMarkersList = TimelineStateA.IsLooping() ? SyncMarkersA : SyncMarkersB;
+				const TArrayView<const FTimelineSyncMarker>& NonLoopingSyncMarkersList = TimelineStateA.IsLooping() ? SyncMarkersB : SyncMarkersA;
+
+				const int32 NumLoopingMarkers = GetNumLoopingMarkers(LoopingTimelineState, LoopingSyncMarkersList);
+				const int32 NumNonLoopingMarkers = NonLoopingSyncMarkersList.Num();
+
+				if (NumNonLoopingMarkers == 0 || NumLoopingMarkers == 0)
 				{
-					ShortestIdx = LongestIdx % NumShortestSyncMarkers;
+					return false;	// Cannot match if we have no markers
 				}
 
-				if ((*LongestSyncMarkersList)[LongestIdx].GetName() != (*ShortestSyncMarkersList)[ShortestIdx].GetName())
+				if ((NumNonLoopingMarkers % NumLoopingMarkers) != 0)
 				{
-					// Markers do not match which means it is not possible for us to phase match
+					// Our looping clip isn't fully looping, we don't handle partial looping
 					return false;
 				}
-			}
 
-			return true;
+				for (int32 NonLoopingMarkerIdx = 0; NonLoopingMarkerIdx < NumNonLoopingMarkers; ++NonLoopingMarkerIdx)
+				{
+					const int32 LoopingMarkerIndex = NonLoopingMarkerIdx % NumLoopingMarkers;
+
+					if (NonLoopingSyncMarkersList[NonLoopingMarkerIdx].GetName() != LoopingSyncMarkersList[LoopingMarkerIndex].GetName())
+					{
+						// Markers do not match, we cannot phase match
+						return false;
+					}
+				}
+
+				// Markers are perfectly in sync
+				return true;
+			}
 		}
 
 		// Returns the phase position for B that matches the current phase on A
-		static FSyncGroupPhasePosition FindMatchingPhaseTime(const TArrayView<const FTimelineSyncMarker>& SyncMarkersA, const FSyncGroupPhasePosition& SyncPositionA, const FTimelineSyncMarkerArray& SyncMarkersB)
+		static FSyncGroupPhasePosition FindMatchingPhaseTime(
+			const FTimelineState& TimelineStateA, const TArrayView<const FTimelineSyncMarker>& SyncMarkersA, const FSyncGroupPhasePosition& SyncPositionA,
+			const FTimelineState& TimelineStateB, const FTimelineSyncMarkerArray& SyncMarkersB)
 		{
-			// We know our phases can match, so we can simply skip the necessary amount of markers in B to match
-			// Consider an example:
-			//   - Sequence A: [R, L, R, L]
-			//   - Sequence B: [R, L, R, L, R, L]
-			// Both A and B can remap with a simply module because the phases must match
-			// (B, 4) maps to (A, 0): 4 % 4 = 0
-			// 
-			// Similarly with more markers:
-			//   - Sequence C: [R, M, L, R, M, L]
-			//   - Sequence D: [R, M, L, R, M, L, R, M, L]
-			// 
-			// We also have the case where an animation loops. Looping animations have the first/last keyframes
-			// identical with the exception of root motion being different. And so if we have a sync marker on
-			// the first keyframe, we need a matching sync marker on the last keyframe (we never interpolate
-			// between the last/first keyframes).
-			//   - Sequence E: [R, L, R]
-			//   - Sequence F: [R, L, R, L, R]
-			// Here, we can't use the modulo as is, we have to add 1 to the marker index after the module operation
-			// (F, 3) maps to (E, 1): (3 % 3) + 1 = 1
-			// Note that this only works if we remap a marker that exceeds the number in the other. Otherwise
-			// we use the index as-is.
+			// If we get here, we know we can phase match, @see CanMatchPhase above for details
 
 			int32 PrevMarkerIndex;
 			int32 NextMarkerIndex;
-			if (SyncMarkersA[0].GetName() == SyncMarkersA.Last().GetName())
+
+			if (!TimelineStateA.IsLooping() && !TimelineStateB.IsLooping())
 			{
-				// Looping markers
-				PrevMarkerIndex = SyncPositionA.PrevMarkerIndex <= SyncMarkersB.Num() ? SyncPositionA.PrevMarkerIndex : ((SyncPositionA.PrevMarkerIndex % SyncMarkersB.Num()) + 1);
-				NextMarkerIndex = SyncPositionA.NextMarkerIndex <= SyncMarkersB.Num() ? SyncPositionA.NextMarkerIndex : ((SyncPositionA.NextMarkerIndex % SyncMarkersB.Num()) + 1);
+				// Neither are looping
+				check(SyncMarkersA.Num() == SyncMarkersB.Num());	// Invariant enforced in CanMatchPhase
+
+				PrevMarkerIndex = SyncPositionA.PrevMarkerIndex;
+				NextMarkerIndex = SyncPositionA.NextMarkerIndex;
+			}
+			else if (TimelineStateA.IsLooping() && TimelineStateB.IsLooping())
+			{
+				// Both are looping
+				const int32 NumSyncMarkersA = GetNumLoopingMarkers(TimelineStateA, SyncMarkersA);
+				const int32 NumSyncMarkersB = GetNumLoopingMarkers(TimelineStateB, SyncMarkersB);
+
+				check(NumSyncMarkersA != 0 && NumSyncMarkersB != 0);	// Invariant enforced in CanMatchPhase
+
+				const int32 NumTargetMarkers = FMath::GreatestCommonDivisor(NumSyncMarkersA, NumSyncMarkersB);
+
+				// TODO: Because we don't keep track of how many loop iterations we've performed, this doesn't quite work
+				//       If the short loop is leading, when it wraps around we'll cause the longer loop to wrap around as well
+
+				PrevMarkerIndex = SyncPositionA.PrevMarkerIndex % NumTargetMarkers;
+				NextMarkerIndex = SyncPositionA.NextMarkerIndex % NumTargetMarkers;
 			}
 			else
 			{
-				PrevMarkerIndex = SyncPositionA.PrevMarkerIndex % SyncMarkersB.Num();
-				NextMarkerIndex = SyncPositionA.NextMarkerIndex % SyncMarkersB.Num();
+				// One of them is looping
+				const FTimelineState& LoopingTimelineState = TimelineStateA.IsLooping() ? TimelineStateA : TimelineStateB;
+				const TArrayView<const FTimelineSyncMarker>& LoopingSyncMarkersList = TimelineStateA.IsLooping() ? SyncMarkersA : SyncMarkersB;
+				const TArrayView<const FTimelineSyncMarker>& NonLoopingSyncMarkersList = TimelineStateA.IsLooping() ? SyncMarkersB : SyncMarkersA;
+
+				const int32 NumLoopingMarkers = GetNumLoopingMarkers(LoopingTimelineState, LoopingSyncMarkersList);
+				const int32 NumNonLoopingMarkers = NonLoopingSyncMarkersList.Num();
+
+				check(NumNonLoopingMarkers != 0 && NumLoopingMarkers != 0);	// Invariant enforced in CanMatchPhase
+				check((NumNonLoopingMarkers % NumLoopingMarkers) == 0);		// Invariant enforced in CanMatchPhase
+
+				// TODO: Because we don't keep track of how many loop iterations we've performed, this doesn't quite work
+				//       If the short loop is leading, when it wraps around we'll cause the longer loop to wrap around as well
+
+				const int32 NumTargetMarkers = FMath::GreatestCommonDivisor(NumLoopingMarkers, NumNonLoopingMarkers);
+
+				PrevMarkerIndex = SyncPositionA.PrevMarkerIndex % NumTargetMarkers;
+				NextMarkerIndex = SyncPositionA.NextMarkerIndex % NumTargetMarkers;
 			}
 
 			// Markers should match
-			check(SyncMarkersA[SyncPositionA.PrevMarkerIndex].GetName() == SyncMarkersB[PrevMarkerIndex].GetName());
-			check(SyncMarkersA[SyncPositionA.NextMarkerIndex].GetName() == SyncMarkersB[NextMarkerIndex].GetName());
+			check(GetMarkerName(SyncPositionA.PrevMarkerIndex, SyncMarkersA) == GetMarkerName(PrevMarkerIndex, SyncMarkersB));
+			check(GetMarkerName(SyncPositionA.NextMarkerIndex, SyncMarkersA) == GetMarkerName(NextMarkerIndex, SyncMarkersB));
 
 			FSyncGroupPhasePosition PhasePosition = SyncPositionA;
 			PhasePosition.PrevMarkerIndex = PrevMarkerIndex;
@@ -1538,7 +1675,10 @@ namespace UE::UAF
 					float DesiredTime;
 
 					// Can we match the phase of the previous leader?
-					const bool bCanMatchPhase = CanMatchPhase(GroupContext.State->PreviousLeaderSyncMarkers, GroupLeaderContext.SyncMarkers);
+					const bool bCanMatchPhase = CanMatchPhase(
+						GroupContext.State->PreviousLeaderTimelineState, GroupContext.State->PreviousLeaderSyncMarkers,
+						LeaderStartTimelineState, GroupLeaderContext.SyncMarkers);
+
 					if (bCanMatchPhase)
 					{
 						// When our phases can match, we simulate playback from the start of the timeline
@@ -1546,7 +1686,9 @@ namespace UE::UAF
 						// We collect the sync markers of our previous leader between [StartPosition, CurrentPosition]
 						// and we pass the same markers on our follower
 
-						const FSyncGroupPhasePosition LeaderPhasePosition = FindMatchingPhaseTime(GroupContext.State->PreviousLeaderSyncMarkers, GroupContext.State->PreviousLeaderPhasePosition, GroupLeaderContext.SyncMarkers);
+						const FSyncGroupPhasePosition LeaderPhasePosition = FindMatchingPhaseTime(
+							GroupContext.State->PreviousLeaderTimelineState, GroupContext.State->PreviousLeaderSyncMarkers, GroupContext.State->PreviousLeaderPhasePosition,
+							LeaderStartTimelineState, GroupLeaderContext.SyncMarkers);
 
 						const FMarkerPair PrevMarker(LeaderPhasePosition.PrevMarkerIndex, 0.0f);
 						const FMarkerPair NextMarker(LeaderPhasePosition.NextMarkerIndex, 0.0f);
@@ -1846,7 +1988,9 @@ namespace UE::UAF
 			FString JoiningMode;
 #endif
 
-			const bool bCanMatchPhase = CanMatchPhase(GroupLeaderContext.SyncMarkers, GroupMemberContext.SyncMarkers);
+			const bool bCanMatchPhase = CanMatchPhase(
+				GroupContext.LeaderTimelineState, GroupLeaderContext.SyncMarkers,
+				FollowerStartTimelineState, GroupMemberContext.SyncMarkers);
 
 			if (GroupMemberContext.State->bJustJoined)
 			{
@@ -1861,7 +2005,9 @@ namespace UE::UAF
 						// We collect the sync markers of our leader between [StartPosition, CurrentPosition]
 						// and we pass the same markers on our follower
 
-						const FSyncGroupPhasePosition FollowerPhasePosition = FindMatchingPhaseTime(GroupLeaderContext.SyncMarkers, GroupContext.LeaderPhaseStart, GroupMemberContext.SyncMarkers);
+						const FSyncGroupPhasePosition FollowerPhasePosition = FindMatchingPhaseTime(
+							GroupContext.LeaderTimelineState, GroupLeaderContext.SyncMarkers, GroupContext.LeaderPhaseStart,
+							FollowerStartTimelineState, GroupMemberContext.SyncMarkers);
 
 						const FMarkerPair PrevMarker(FollowerPhasePosition.PrevMarkerIndex, 0.0f);
 						const FMarkerPair NextMarker(FollowerPhasePosition.NextMarkerIndex, 0.0f);
