Add new UAF trait to apply a named set

#Synced-CL 48995018

diff --git a/Engine/Plugins/Experimental/UAF/UAFAnimGraph/Source/UAFAnimGraph/Private/EvaluationVM/Tasks/ApplyNamedSetTask.cpp b/Engine/Plugins/Experimental/UAF/UAFAnimGraph/Source/UAFAnimGraph/Private/EvaluationVM/Tasks/ApplyNamedSetTask.cpp
new file mode 100644
--- /dev/null
+++ b/Engine/Plugins/Experimental/UAF/UAFAnimGraph/Source/UAFAnimGraph/Private/EvaluationVM/Tasks/ApplyNamedSetTask.cpp
@@ -0,0 +1,56 @@
+// Copyright Epic Games, Inc. All Rights Reserved.
+
+#include "EvaluationVM/Tasks/ApplyNamedSetTask.h"
+
+#include "EvaluationVM/EvaluationVM.h"
+
+#include UE_INLINE_GENERATED_CPP_BY_NAME(ApplyNamedSetTask)
+
+namespace UE::UAF
+{
+	FPushNamedSetTask FPushNamedSetTask::Make(FName SetName)
+	{
+		FPushNamedSetTask Task;
+		Task.SetName = SetName;
+		return Task;
+	}
+
+	void FPushNamedSetTask::Execute(UE::UAF::FEvaluationVM& VM) const
+	{
+		// If we use the old runtime, do nothing
+		if (!VM.GetActiveNamedSet())
+		{
+			return;
+		}
+
+		if (!VM.PushEvaluationContext(SetName))
+		{
+			UE_LOG(LogAnimation, Warning, TEXT("FPushNamedSetTask::Execute: Could not push new Named Set, set was not found in current binding."));
+		}
+	}
+
+	FPopNamedSetTask FPopNamedSetTask::Make(FName ExpectedSetName)
+	{
+		FPopNamedSetTask Task;
+		Task.ExpectedSetName = ExpectedSetName;
+		return Task;
+	}
+
+	void FPopNamedSetTask::Execute(UE::UAF::FEvaluationVM& VM) const
+	{
+		// If we use the old runtime, do nothing
+		if (!VM.GetActiveNamedSet())
+		{
+			return;
+		}
+
+		const FName ActiveSetName = VM.GetActiveEvaluationContext().GetNamedSet()->GetName();
+		if (ActiveSetName != ExpectedSetName)
+		{
+			UE_LOG(LogAnimation, Warning, TEXT("FPopNamedSetTask::Execute: Could not pop active Named Set, the expected set to pop was not active."));
+			return;
+		}
+
+		VM.PopEvaluationContext();
+	}
+}
diff --git a/Engine/Plugins/Experimental/UAF/UAFAnimGraph/Source/UAFAnimGraph/Private/Traits/ApplyNamedSetTrait.cpp b/Engine/Plugins/Experimental/UAF/UAFAnimGraph/Source/UAFAnimGraph/Private/Traits/ApplyNamedSetTrait.cpp
new file mode 100644
--- /dev/null
+++ b/Engine/Plugins/Experimental/UAF/UAFAnimGraph/Source/UAFAnimGraph/Private/Traits/ApplyNamedSetTrait.cpp
@@ -0,0 +1,64 @@
+// Copyright Epic Games, Inc. All Rights Reserved.
+
+#include "Traits/ApplyNamedSetTrait.h"
+
+#include "EvaluationVM/Tasks/ApplyNamedSetTask.h"
+
+namespace UE::UAF
+{
+	AUTO_REGISTER_ANIM_TRAIT(FApplyNamedSetTrait)
+
+	// Trait implementation boilerplate
+	#define TRAIT_INTERFACE_ENUMERATOR(GeneratorMacro) \
+		GeneratorMacro(IEvaluate) \
+		GeneratorMacro(IHierarchy) \
+
+	GENERATE_ANIM_TRAIT_IMPLEMENTATION(FApplyNamedSetTrait, TRAIT_INTERFACE_ENUMERATOR, NULL_ANIM_TRAIT_INTERFACE_ENUMERATOR, NULL_ANIM_TRAIT_EVENT_ENUMERATOR)
+	#undef TRAIT_INTERFACE_ENUMERATOR
+
+	void FApplyNamedSetTrait::FInstanceData::Construct(const FExecutionContext& Context, const FTraitBinding& Binding)
+	{
+		FTrait::FInstanceData::Construct(Context, Binding);
+
+		const FSharedData* SharedData = Binding.GetSharedData<FSharedData>();
+		FInstanceData* InstanceData = Binding.GetInstanceData<FInstanceData>();
+
+		if (!InstanceData->Input.IsValid())
+		{
+			InstanceData->Input = Context.AllocateNodeInstance(Binding, SharedData->Input);
+		}
+	}
+
+	void FApplyNamedSetTrait::PreEvaluate(FEvaluateTraversalContext& Context, const TTraitBinding<IEvaluate>& Binding) const
+	{
+		const FSharedData* SharedData = Binding.GetSharedData<FSharedData>();
+		const FName SetName = SharedData->GetSetName(Binding);
+
+		Context.AppendTask(FPushNamedSetTask::Make(SetName));
+
+		IEvaluate::PreEvaluate(Context, Binding);
+	}
+
+	void FApplyNamedSetTrait::PostEvaluate(FEvaluateTraversalContext& Context, const TTraitBinding<IEvaluate>& Binding) const
+	{
+		IEvaluate::PostEvaluate(Context, Binding);
+
+		const FSharedData* SharedData = Binding.GetSharedData<FSharedData>();
+		const FName SetName = SharedData->GetSetName(Binding);
+
+		Context.AppendTask(FPopNamedSetTask::Make(SetName));
+	}
+
+	uint32 FApplyNamedSetTrait::GetNumChildren(const FExecutionContext& Context, const TTraitBinding<IHierarchy>& Binding) const
+	{
+		return 1;
+	}
+
+	void FApplyNamedSetTrait::GetChildren(const FExecutionContext& Context, const TTraitBinding<IHierarchy>& Binding, FChildrenArray& Children) const
+	{
+		const FInstanceData* InstanceData = Binding.GetInstanceData<FInstanceData>();
+
+		// Add the child, even if the handle is empty
+		Children.Add(InstanceData->Input);
+	}
+}
diff --git a/Engine/Plugins/Experimental/UAF/UAFAnimGraph/Source/UAFAnimGraph/Private/Traits/ApplyNamedSetTrait.h b/Engine/Plugins/Experimental/UAF/UAFAnimGraph/Source/UAFAnimGraph/Private/Traits/ApplyNamedSetTrait.h
new file mode 100644
--- /dev/null
+++ b/Engine/Plugins/Experimental/UAF/UAFAnimGraph/Source/UAFAnimGraph/Private/Traits/ApplyNamedSetTrait.h
@@ -0,0 +1,42 @@
+// Copyright Epic Games, Inc. All Rights Reserved.
+
+#pragma once
+
+#include "CoreMinimal.h"
+
+#include "TraitCore/Trait.h"
+#include "TraitInterfaces/IEvaluate.h"
+#include "TraitInterfaces/IHierarchy.h"
+#include "Traits/ApplyNamedSetTraitData.h"
+
+namespace UE::UAF
+{
+	/**
+	 * FApplyNamedSetTrait
+	 *
+	 * A trait that applies a new named set to evaluate with.
+	 * Traits upstream will use the new named set to evaluate.
+	 * The trait that consumes its output must handle mismatched named sets (e.g. layering).
+	 */
+	struct FApplyNamedSetTrait : FBaseTrait, IEvaluate, IHierarchy
+	{
+		DECLARE_ANIM_TRAIT(FApplyNamedSetTrait, FBaseTrait)
+
+		using FSharedData = FApplyNamedSetSharedData;
+
+		struct FInstanceData : FTrait::FInstanceData
+		{
+			FTraitPtr Input;
+
+			void Construct(const FExecutionContext& Context, const FTraitBinding& Binding);
+		};
+
+		// IEvaluate impl
+		virtual void PreEvaluate(FEvaluateTraversalContext& Context, const TTraitBinding<IEvaluate>& Binding) const override;
+		virtual void PostEvaluate(FEvaluateTraversalContext& Context, const TTraitBinding<IEvaluate>& Binding) const override;
+
+		// IHierarchy impl
+		virtual uint32 GetNumChildren(const FExecutionContext& Context, const TTraitBinding<IHierarchy>& Binding) const override;
+		virtual void GetChildren(const FExecutionContext& Context, const TTraitBinding<IHierarchy>& Binding, FChildrenArray& Children) const override;
+	};
+}
diff --git a/Engine/Plugins/Experimental/UAF/UAFAnimGraph/Source/UAFAnimGraph/Public/EvaluationVM/Tasks/ApplyNamedSetTask.h b/Engine/Plugins/Experimental/UAF/UAFAnimGraph/Source/UAFAnimGraph/Public/EvaluationVM/Tasks/ApplyNamedSetTask.h
new file mode 100644
--- /dev/null
+++ b/Engine/Plugins/Experimental/UAF/UAFAnimGraph/Source/UAFAnimGraph/Public/EvaluationVM/Tasks/ApplyNamedSetTask.h
@@ -0,0 +1,58 @@
+// Copyright Epic Games, Inc. All Rights Reserved.
+
+#pragma once
+
+#include "CoreMinimal.h"
+#include "EvaluationVM/EvaluationTask.h"
+
+#include "ApplyNamedSetTask.generated.h"
+
+#define UE_API UAFANIMGRAPH_API
+
+namespace UE::UAF
+{
+	/*
+	 * Push Named Set Task
+	 *
+	 * This pushes a new named set to evaluate with. Subsequent tasks will use that set to evaluate.
+	 */
+	USTRUCT()
+	struct FPushNamedSetTask : public FAnimNextEvaluationTask
+	{
+		GENERATED_BODY()
+
+		DECLARE_ANIM_EVALUATION_TASK(FPushNamedSetTask)
+
+		static UE_API FPushNamedSetTask Make(FName SetName);
+
+		// Task entry point
+		UE_API virtual void Execute(UE::UAF::FEvaluationVM& VM) const override;
+
+		// The named set to push
+		FName SetName = NAME_None;
+	};
+
+	/*
+	 * Pop Named Set Task
+	 *
+	 * This pops the current active named set. Subsequent tasks will no longer use that named set.
+	 * We only pop the active named set if it is the one we expect.
+	 */
+	USTRUCT()
+	struct FPopNamedSetTask : public FAnimNextEvaluationTask
+	{
+		GENERATED_BODY()
+
+		DECLARE_ANIM_EVALUATION_TASK(FPopNamedSetTask)
+
+		static UE_API FPopNamedSetTask Make(FName ExpectedSetName);
+
+		// Task entry point
+		UE_API virtual void Execute(UE::UAF::FEvaluationVM& VM) const override;
+
+		// The named set expected to pop
+		FName ExpectedSetName = NAME_None;
+	};
+}
+
+#undef UE_API
diff --git a/Engine/Plugins/Experimental/UAF/UAFAnimGraph/Source/UAFAnimGraph/Public/Traits/ApplyNamedSetTraitData.h b/Engine/Plugins/Experimental/UAF/UAFAnimGraph/Source/UAFAnimGraph/Public/Traits/ApplyNamedSetTraitData.h
new file mode 100644
--- /dev/null
+++ b/Engine/Plugins/Experimental/UAF/UAFAnimGraph/Source/UAFAnimGraph/Public/Traits/ApplyNamedSetTraitData.h
@@ -0,0 +1,39 @@
+// Copyright Epic Games, Inc. All Rights Reserved.
+
+#pragma once
+
+#include "CoreMinimal.h"
+
+#include "TraitCore/TraitBinding.h"
+#include "TraitCore/TraitHandle.h"
+#include "TraitCore/TraitSharedData.h"
+
+#include "ApplyNamedSetTraitData.generated.h"
+
+namespace UE::UAF
+{
+	/**
+	 * A trait that applies a new named set to evaluate with.
+	 * Traits upstream will use the new named set to evaluate.
+	 * The trait that consumes its output must handle mismatched named sets (e.g. layering).
+	 */
+	USTRUCT(meta = (DisplayName = "Apply Named Set", ShowTooltip=true))
+	struct FApplyNamedSetSharedData : public FAnimNextTraitSharedData
+	{
+		GENERATED_BODY()
+
+		/** Input to evaluate with the specified set name */
+		UPROPERTY()
+		FAnimNextTraitHandle Input;
+
+		/** Named set to evaluate our children with */
+		UPROPERTY(Category=Settings, EditAnywhere)
+		FName SetName = NAME_None;
+
+		#define TRAIT_LATENT_PROPERTIES_ENUMERATOR(GeneratorMacro) \
+			GeneratorMacro(SetName) \
+
+		GENERATE_TRAIT_LATENT_PROPERTIES(FApplyNamedSetSharedData, TRAIT_LATENT_PROPERTIES_ENUMERATOR)
+		#undef TRAIT_LATENT_PROPERTIES_ENUMERATOR
+	};
+}
