UAF: Fix race condition with graph evaluator method registration

Add a lock as the map can be contended in builds with async loading enabled

#Synced-CL 49075355

diff --git a/Engine/Plugins/Experimental/UAF/UAFAnimGraph/Source/UAFAnimGraph/Private/Graph/RigUnit_AnimNextGraphEvaluator.cpp b/Engine/Plugins/Experimental/UAF/UAFAnimGraph/Source/UAFAnimGraph/Private/Graph/RigUnit_AnimNextGraphEvaluator.cpp
--- a/Engine/Plugins/Experimental/UAF/UAFAnimGraph/Source/UAFAnimGraph/Private/Graph/RigUnit_AnimNextGraphEvaluator.cpp
+++ b/Engine/Plugins/Experimental/UAF/UAFAnimGraph/Source/UAFAnimGraph/Private/Graph/RigUnit_AnimNextGraphEvaluator.cpp
@@ -10,6 +10,7 @@
 
 namespace UE::UAF::Private
 {
+	static FRWLock GRegisteredGraphEvaluatorMethodsLock;
 	static TMap<uint32, FAnimNextGraphEvaluatorExecuteDefinition> GRegisteredGraphEvaluatorMethods;
 
 	TArray<FRigVMFunctionArgument> GetGraphEvaluatorFunctionArguments(const FAnimNextGraphEvaluatorExecuteDefinition& ExecuteDefinition)
@@ -80,12 +81,15 @@ void FRigUnit_AnimNextGraphEvaluator::RegisterExecuteMethod(const FAnimNextGraph
 {
 	using namespace UE::UAF::Private;
 
-	if (GRegisteredGraphEvaluatorMethods.Contains(ExecuteDefinition.Hash))
 	{
-		return;	// Already registered
-	}
+		UE::TWriteScopeLock WriteLock(GRegisteredGraphEvaluatorMethodsLock);
+		if (GRegisteredGraphEvaluatorMethods.Contains(ExecuteDefinition.Hash))
+		{
+			return;	// Already registered
+		}
 
-	GRegisteredGraphEvaluatorMethods.Add(ExecuteDefinition.Hash, ExecuteDefinition);
+		GRegisteredGraphEvaluatorMethods.Add(ExecuteDefinition.Hash, ExecuteDefinition);
+	}
 
 	const FString FullExecuteMethodName = FString::Printf(TEXT("FRigUnit_AnimNextGraphEvaluator::%s"), *ExecuteDefinition.MethodName);
 
@@ -97,5 +101,6 @@ const FAnimNextGraphEvaluatorExecuteDefinition* FRigUnit_AnimNextGraphEvaluator:
 {
 	using namespace UE::UAF::Private;
 
+	UE::TWriteScopeLock ReadLock(GRegisteredGraphEvaluatorMethodsLock);
 	return GRegisteredGraphEvaluatorMethods.Find(ExecuteMethodHash);
 }
