UAF: Split up traits and their shared data into separate headers and move shared data headers into public to be visible to external code

#Synced-CL 49202942

diff --git a/Engine/Plugins/Experimental/UAF/UAFAnimGraph/Source/UAFAnimGraph/Private/Traits/ApplyAdditiveTrait.cpp b/Engine/Plugins/Experimental/UAF/UAFAnimGraph/Source/UAFAnimGraph/Private/Traits/ApplyAdditiveTrait.cpp
--- a/Engine/Plugins/Experimental/UAF/UAFAnimGraph/Source/UAFAnimGraph/Private/Traits/ApplyAdditiveTrait.cpp
+++ b/Engine/Plugins/Experimental/UAF/UAFAnimGraph/Source/UAFAnimGraph/Private/Traits/ApplyAdditiveTrait.cpp
@@ -7,8 +7,6 @@
 #include "TraitCore/ExecutionContext.h"
 #include "EvaluationVM/Tasks/ApplyAdditiveKeyframe.h"
 
-#include UE_INLINE_GENERATED_CPP_BY_NAME(ApplyAdditiveTrait)
-
 namespace UE::UAF
 {
 	AUTO_REGISTER_ANIM_TRAIT(FApplyAdditiveTrait)
diff --git a/Engine/Plugins/Experimental/UAF/UAFAnimGraph/Source/UAFAnimGraph/Private/Traits/ApplyAdditiveTrait.h b/Engine/Plugins/Experimental/UAF/UAFAnimGraph/Source/UAFAnimGraph/Private/Traits/ApplyAdditiveTrait.h
--- a/Engine/Plugins/Experimental/UAF/UAFAnimGraph/Source/UAFAnimGraph/Private/Traits/ApplyAdditiveTrait.h
+++ b/Engine/Plugins/Experimental/UAF/UAFAnimGraph/Source/UAFAnimGraph/Private/Traits/ApplyAdditiveTrait.h
@@ -10,42 +10,7 @@
 #include "TraitInterfaces/IEvaluate.h"
 #include "TraitInterfaces/IHierarchy.h"
 #include "TraitInterfaces/IUpdate.h"
-
-#include "ApplyAdditiveTrait.generated.h"
-
-/**
- * A trait that can apply a mesh or local space additive to this trait stack.
- * Ex: LookAt's, minor hit reactions, up-down floating, etc.
- */
-USTRUCT(meta = (DisplayName = "Apply Additive", ShowTooltip=true))
-struct FAnimNextApplyAdditiveTraitSharedData : public FAnimNextTraitSharedData
-{
-	GENERATED_BODY()
-
-	/** Base to apply addtitive too. */
-	UPROPERTY()
-	FAnimNextTraitHandle Base;
-
-	/** Addtive to apply. */
-	UPROPERTY()
-	FAnimNextTraitHandle Additive;
-
-	/** 
-	 * Deprecated. Please add an IAlphaInputArgs addtive trait to set alpha.
-	 * @TODO: Remove pre 5.6 once removing latents doesn't cause crash
-	 * 
-	 * How much to apply our addtive, default is 1. 
-	 */
-	UPROPERTY(EditAnywhere, Category = "Default")
-	float Alpha = 1.0f;
-
-	// Latent pin support boilerplate
-	#define TRAIT_LATENT_PROPERTIES_ENUMERATOR(GeneratorMacro) \
-		GeneratorMacro(Alpha) \
-
-	GENERATE_TRAIT_LATENT_PROPERTIES(FAnimNextApplyAdditiveTraitSharedData, TRAIT_LATENT_PROPERTIES_ENUMERATOR)
-	#undef TRAIT_LATENT_PROPERTIES_ENUMERATOR
-};
+#include "Traits/ApplyAdditiveTraitData.h"
 
 namespace UE::UAF
 {
@@ -65,7 +30,7 @@ namespace UE::UAF
 			/** Reference to base for additive we are applying */
 			FTraitPtr Base;
 
-			/** Reference to addtive we are applying */
+			/** Reference to additive we are applying */
 			FTraitPtr Additive;
 
 			/** True if Additive branch has any contribution (Alpha non-zero)  */
diff --git a/Engine/Plugins/Experimental/UAF/UAFAnimGraph/Source/UAFAnimGraph/Private/Traits/BlendLayer.cpp b/Engine/Plugins/Experimental/UAF/UAFAnimGraph/Source/UAFAnimGraph/Private/Traits/BlendLayer.cpp
--- a/Engine/Plugins/Experimental/UAF/UAFAnimGraph/Source/UAFAnimGraph/Private/Traits/BlendLayer.cpp
+++ b/Engine/Plugins/Experimental/UAF/UAFAnimGraph/Source/UAFAnimGraph/Private/Traits/BlendLayer.cpp
@@ -13,7 +13,6 @@
 #include "MaskProfile/HierarchyTableTypeMask.h"
 #include "SkeletonHierarchyTableType.h"
 
-#include UE_INLINE_GENERATED_CPP_BY_NAME(BlendLayer)
 
 class FHierarchyTableBlendProfileProvider : public IInterpolationIndexProvider
 {
diff --git a/Engine/Plugins/Experimental/UAF/UAFAnimGraph/Source/UAFAnimGraph/Private/Traits/BlendLayer.h b/Engine/Plugins/Experimental/UAF/UAFAnimGraph/Source/UAFAnimGraph/Private/Traits/BlendLayer.h
--- a/Engine/Plugins/Experimental/UAF/UAFAnimGraph/Source/UAFAnimGraph/Private/Traits/BlendLayer.h
+++ b/Engine/Plugins/Experimental/UAF/UAFAnimGraph/Source/UAFAnimGraph/Private/Traits/BlendLayer.h
@@ -5,47 +5,12 @@
 #include "CoreMinimal.h"
 
 #include "EvaluationVM/Tasks/BlendKeyframesPerBone.h"
-#include "HierarchyTable.h"
 #include "TraitCore/Trait.h"
 #include "TraitInterfaces/IContinuousBlend.h"
 #include "TraitInterfaces/IEvaluate.h"
 #include "TraitInterfaces/IHierarchy.h"
 #include "TraitInterfaces/IUpdate.h"
-
-#include "BlendLayer.generated.h"
-
-/** A trait that can blend a layer into an input. */
-USTRUCT(meta = (DisplayName = "Blend Layer", ShowTooltip=true))
-struct FAnimNextBlendLayerTraitSharedData : public FAnimNextTraitSharedData
-{
-	GENERATED_BODY()
-
-	UPROPERTY()
-	FAnimNextTraitHandle ChildBase;
-
-	UPROPERTY()
-	FAnimNextTraitHandle ChildBlend;
-
-	// The strength with which to apply the blend pose
-	UPROPERTY(EditAnywhere, Category = "Default")
-	float BlendWeight = 1.0f;
-
-	/** Blend profile that configures how fast to blend each bone. */
-	// TODO: Can't show list of blend profiles, we need to find a skeleton to perform the lookup with
-	UPROPERTY(EditAnywhere, Category = "Default", meta = (Inline))
-	TObjectPtr<UHierarchyTable> BlendProfile = nullptr;
-
-	// Whether or not to blend in mesh space or local space
-	UPROPERTY(EditAnywhere, Category = "Default")
-	bool bBlendInMeshSpace = false;
-
-	#define TRAIT_LATENT_PROPERTIES_ENUMERATOR(GeneratorMacro) \
-		GeneratorMacro(BlendWeight) \
-		GeneratorMacro(bBlendInMeshSpace) \
-
-	GENERATE_TRAIT_LATENT_PROPERTIES(FAnimNextBlendLayerTraitSharedData, TRAIT_LATENT_PROPERTIES_ENUMERATOR)
-	#undef TRAIT_LATENT_PROPERTIES_ENUMERATOR
-};
+#include "Traits/BlendLayerTraitData.h"
 
 namespace UE::UAF
 {
diff --git a/Engine/Plugins/Experimental/UAF/UAFAnimGraph/Source/UAFAnimGraph/Private/Traits/BlendTwoWay.cpp b/Engine/Plugins/Experimental/UAF/UAFAnimGraph/Source/UAFAnimGraph/Private/Traits/BlendTwoWay.cpp
--- a/Engine/Plugins/Experimental/UAF/UAFAnimGraph/Source/UAFAnimGraph/Private/Traits/BlendTwoWay.cpp
+++ b/Engine/Plugins/Experimental/UAF/UAFAnimGraph/Source/UAFAnimGraph/Private/Traits/BlendTwoWay.cpp
@@ -6,8 +6,6 @@
 #include "TraitCore/ExecutionContext.h"
 #include "EvaluationVM/Tasks/BlendKeyframes.h"
 
-#include UE_INLINE_GENERATED_CPP_BY_NAME(BlendTwoWay)
-
 namespace UE::UAF
 {
 	AUTO_REGISTER_ANIM_TRAIT(FBlendTwoWayTrait)
diff --git a/Engine/Plugins/Experimental/UAF/UAFAnimGraph/Source/UAFAnimGraph/Private/Traits/BlendTwoWay.h b/Engine/Plugins/Experimental/UAF/UAFAnimGraph/Source/UAFAnimGraph/Private/Traits/BlendTwoWay.h
--- a/Engine/Plugins/Experimental/UAF/UAFAnimGraph/Source/UAFAnimGraph/Private/Traits/BlendTwoWay.h
+++ b/Engine/Plugins/Experimental/UAF/UAFAnimGraph/Source/UAFAnimGraph/Private/Traits/BlendTwoWay.h
@@ -9,44 +9,7 @@
 #include "TraitInterfaces/IEvaluate.h"
 #include "TraitInterfaces/IHierarchy.h"
 #include "TraitInterfaces/IUpdate.h"
-
-#include "BlendTwoWay.generated.h"
-
-/** A trait that can blend two inputs. */
-USTRUCT(meta = (DisplayName = "Blend Two Way", ShowTooltip=true))
-struct FAnimNextBlendTwoWayTraitSharedData : public FAnimNextTraitSharedData
-{
-	GENERATED_BODY()
-
-	/** First output to be blended (full weight is 0.0). */
-	UPROPERTY()
-	FAnimNextTraitHandle ChildA;
-
-	/** Second output to be blended (full weight is 1.0). */
-	UPROPERTY()
-	FAnimNextTraitHandle ChildB;
-
-	/** How much to blend our two children: 0.0 is fully child A while 1.0 is fully child B. */
-	UPROPERTY(EditAnywhere, Category = "Default")
-	float BlendWeight = 0.0f;
-
-	/** This reinitializes child when re-activated */
-	UPROPERTY(EditAnywhere, Category = "Default")
-	bool bResetChildOnActivation = true;
-
-	/** Always update children, regardless of whether or not that child has weight. */
-	UPROPERTY(EditAnywhere, Category = "Default")
-	bool bAlwaysUpdateChildren = false;
-
-	// Latent pin support boilerplate
-	#define TRAIT_LATENT_PROPERTIES_ENUMERATOR(GeneratorMacro) \
-		GeneratorMacro(BlendWeight) \
-		GeneratorMacro(bResetChildOnActivation) \
-		GeneratorMacro(bAlwaysUpdateChildren) \
-
-	GENERATE_TRAIT_LATENT_PROPERTIES(FAnimNextBlendTwoWayTraitSharedData, TRAIT_LATENT_PROPERTIES_ENUMERATOR)
-	#undef TRAIT_LATENT_PROPERTIES_ENUMERATOR
-};
+#include "Traits/BlendTwoWayTraitData.h"
 
 namespace UE::UAF
 {
diff --git a/Engine/Plugins/Experimental/UAF/UAFAnimGraph/Source/UAFAnimGraph/Private/Traits/Passthrough.cpp b/Engine/Plugins/Experimental/UAF/UAFAnimGraph/Source/UAFAnimGraph/Private/Traits/Passthrough.cpp
--- a/Engine/Plugins/Experimental/UAF/UAFAnimGraph/Source/UAFAnimGraph/Private/Traits/Passthrough.cpp
+++ b/Engine/Plugins/Experimental/UAF/UAFAnimGraph/Source/UAFAnimGraph/Private/Traits/Passthrough.cpp
@@ -4,8 +4,6 @@
 
 #include "TraitCore/ExecutionContext.h"
 
-#include UE_INLINE_GENERATED_CPP_BY_NAME(Passthrough)
-
 namespace UE::UAF
 {
 	AUTO_REGISTER_ANIM_TRAIT(FPassthroughTrait)
diff --git a/Engine/Plugins/Experimental/UAF/UAFAnimGraph/Source/UAFAnimGraph/Private/Traits/Passthrough.h b/Engine/Plugins/Experimental/UAF/UAFAnimGraph/Source/UAFAnimGraph/Private/Traits/Passthrough.h
--- a/Engine/Plugins/Experimental/UAF/UAFAnimGraph/Source/UAFAnimGraph/Private/Traits/Passthrough.h
+++ b/Engine/Plugins/Experimental/UAF/UAFAnimGraph/Source/UAFAnimGraph/Private/Traits/Passthrough.h
@@ -7,19 +7,7 @@
 #include "TraitCore/Trait.h"
 #include "TraitInterfaces/IHierarchy.h"
 #include "TraitInterfaces/IUpdate.h"
-
-#include "Passthrough.generated.h"
-
-/** A trait that passes through the input without modification. */
-USTRUCT(meta = (DisplayName = "Passthrough", ShowTooltip=true))
-struct FAnimNextPassthroughSharedData : public FAnimNextTraitSharedData
-{
-	GENERATED_BODY()
-
-	/** Input to pass to output */
-	UPROPERTY()
-	FAnimNextTraitHandle Input;
-};
+#include "Traits/PassthroughTraitData.h"
 
 namespace UE::UAF
 {
diff --git a/Engine/Plugins/Experimental/UAF/UAFAnimGraph/Source/UAFAnimGraph/Private/Traits/SubGraphHost.cpp b/Engine/Plugins/Experimental/UAF/UAFAnimGraph/Source/UAFAnimGraph/Private/Traits/SubGraphHost.cpp
--- a/Engine/Plugins/Experimental/UAF/UAFAnimGraph/Source/UAFAnimGraph/Private/Traits/SubGraphHost.cpp
+++ b/Engine/Plugins/Experimental/UAF/UAFAnimGraph/Source/UAFAnimGraph/Private/Traits/SubGraphHost.cpp
@@ -6,8 +6,6 @@
 #include "TraitCore/NodeInstance.h"
 #include "Graph/AnimNextGraphInstance.h"
 
-#include UE_INLINE_GENERATED_CPP_BY_NAME(SubGraphHost)
-
 namespace UE::UAF
 {
 	AUTO_REGISTER_ANIM_TRAIT(FSubGraphHostTrait)
diff --git a/Engine/Plugins/Experimental/UAF/UAFAnimGraph/Source/UAFAnimGraph/Private/Traits/SubGraphHost.h b/Engine/Plugins/Experimental/UAF/UAFAnimGraph/Source/UAFAnimGraph/Private/Traits/SubGraphHost.h
--- a/Engine/Plugins/Experimental/UAF/UAFAnimGraph/Source/UAFAnimGraph/Private/Traits/SubGraphHost.h
+++ b/Engine/Plugins/Experimental/UAF/UAFAnimGraph/Source/UAFAnimGraph/Private/Traits/SubGraphHost.h
@@ -11,35 +11,7 @@
 #include "TraitInterfaces/IUpdate.h"
 #include "Graph/AnimNextAnimationGraph.h"
 #include "Graph/AnimNextGraphInstance.h"
-
-#include "SubGraphHost.generated.h"
-
-/** A trait that hosts and manages a sub-graph instance. */
-USTRUCT(meta = (DisplayName = "SubGraph", ShowTooltip=true))
-struct FAnimNextSubGraphHostTraitSharedData : public FAnimNextTraitSharedData
-{
-	GENERATED_BODY()
-
-	/** Asset to use as a sub-graph */
-	UPROPERTY(EditAnywhere, Category = "Default", meta=(ExportAsReference="true"))
-	TObjectPtr<const UAnimNextAnimationGraph> AnimationGraph;
-
-	/** A dummy child that we can use to output the bind pose. This property is hidden and automatically populated during compilation. */
-	UPROPERTY(meta = (Hidden))
-	FAnimNextTraitHandle ReferencePoseChild;
-
-	/** Entry point in the Subgraph that we will use */
-	UPROPERTY(EditAnywhere, Category = "Default", meta=(CustomWidget = "ParamName", AllowedParamType = "FAnimNextEntryPoint"))
-	FName EntryPoint;
-
-	// Latent pin support boilerplate
-	#define TRAIT_LATENT_PROPERTIES_ENUMERATOR(GeneratorMacro) \
-		GeneratorMacro(AnimationGraph) \
-		GeneratorMacro(EntryPoint) \
-
-	GENERATE_TRAIT_LATENT_PROPERTIES(FAnimNextSubGraphHostTraitSharedData, TRAIT_LATENT_PROPERTIES_ENUMERATOR)
-	#undef TRAIT_LATENT_PROPERTIES_ENUMERATOR
-};
+#include "Traits/SubGraphHostTraitData.h"
 
 namespace UE::UAF
 {
diff --git a/Engine/Plugins/Experimental/UAF/UAFAnimGraph/Source/UAFAnimGraph/Public/Traits/ApplyAdditiveTraitData.h b/Engine/Plugins/Experimental/UAF/UAFAnimGraph/Source/UAFAnimGraph/Public/Traits/ApplyAdditiveTraitData.h
new file mode 100644
--- /dev/null
+++ b/Engine/Plugins/Experimental/UAF/UAFAnimGraph/Source/UAFAnimGraph/Public/Traits/ApplyAdditiveTraitData.h
@@ -0,0 +1,43 @@
+// Copyright Epic Games, Inc. All Rights Reserved.
+
+#pragma once
+
+#include "TraitCore/TraitBinding.h"
+#include "TraitCore/TraitHandle.h"
+#include "TraitCore/TraitSharedData.h"
+
+#include "ApplyAdditiveTraitData.generated.h"
+
+/**
+ * A trait that can apply a mesh or local space additive to this trait stack.
+ * Ex: LookAt's, minor hit reactions, up-down floating, etc.
+ */
+USTRUCT(meta = (DisplayName = "Apply Additive", ShowTooltip=true))
+struct FAnimNextApplyAdditiveTraitSharedData : public FAnimNextTraitSharedData
+{
+	GENERATED_BODY()
+
+	/** Base to apply additive too. */
+	UPROPERTY()
+	FAnimNextTraitHandle Base;
+
+	/** Additive to apply. */
+	UPROPERTY()
+	FAnimNextTraitHandle Additive;
+
+	/** 
+	 * Deprecated. Please add an IAlphaInputArgs additive trait to set alpha.
+	 * @TODO: Remove pre 5.6 once removing latents doesn't cause crash
+	 * 
+	 * How much to apply our additive, default is 1. 
+	 */
+	UPROPERTY(EditAnywhere, Category = "Default")
+	float Alpha = 1.0f;
+	
+	// Latent pin support boilerplate
+	#define TRAIT_LATENT_PROPERTIES_ENUMERATOR(GeneratorMacro) \
+	GeneratorMacro(Alpha) \
+
+	GENERATE_TRAIT_LATENT_PROPERTIES(FAnimNextApplyAdditiveTraitSharedData, TRAIT_LATENT_PROPERTIES_ENUMERATOR)
+	#undef TRAIT_LATENT_PROPERTIES_ENUMERATOR
+};
\ No newline at end of file
diff --git a/Engine/Plugins/Experimental/UAF/UAFAnimGraph/Source/UAFAnimGraph/Public/Traits/BlendLayerTraitData.h b/Engine/Plugins/Experimental/UAF/UAFAnimGraph/Source/UAFAnimGraph/Public/Traits/BlendLayerTraitData.h
new file mode 100644
--- /dev/null
+++ b/Engine/Plugins/Experimental/UAF/UAFAnimGraph/Source/UAFAnimGraph/Public/Traits/BlendLayerTraitData.h
@@ -0,0 +1,44 @@
+// Copyright Epic Games, Inc. All Rights Reserved.
+
+#pragma once
+
+#include "TraitCore/TraitBinding.h"
+#include "TraitCore/TraitHandle.h"
+#include "TraitCore/TraitSharedData.h"
+
+#include "BlendLayerTraitData.generated.h"
+
+class UHierarchyTable;
+
+/** A trait that can blend a layer into an input. */
+USTRUCT(meta = (DisplayName = "Blend Layer", ShowTooltip=true))
+struct FAnimNextBlendLayerTraitSharedData : public FAnimNextTraitSharedData
+{
+	GENERATED_BODY()
+
+	UPROPERTY()
+	FAnimNextTraitHandle ChildBase;
+
+	UPROPERTY()
+	FAnimNextTraitHandle ChildBlend;
+
+	// The strength with which to apply the blend pose
+	UPROPERTY(EditAnywhere, Category = "Default")
+	float BlendWeight = 1.0f;
+
+	/** Blend profile that configures how fast to blend each bone. */
+	// TODO: Can't show list of blend profiles, we need to find a skeleton to perform the lookup with
+	UPROPERTY(EditAnywhere, Category = "Default", meta = (Inline))
+	TObjectPtr<UHierarchyTable> BlendProfile = nullptr;
+
+	// Whether or not to blend in mesh space or local space
+	UPROPERTY(EditAnywhere, Category = "Default")
+	bool bBlendInMeshSpace = false;
+
+	#define TRAIT_LATENT_PROPERTIES_ENUMERATOR(GeneratorMacro) \
+	GeneratorMacro(BlendWeight) \
+	GeneratorMacro(bBlendInMeshSpace) \
+
+	GENERATE_TRAIT_LATENT_PROPERTIES(FAnimNextBlendLayerTraitSharedData, TRAIT_LATENT_PROPERTIES_ENUMERATOR)
+	#undef TRAIT_LATENT_PROPERTIES_ENUMERATOR
+};
\ No newline at end of file
diff --git a/Engine/Plugins/Experimental/UAF/UAFAnimGraph/Source/UAFAnimGraph/Public/Traits/BlendTwoWayTraitData.h b/Engine/Plugins/Experimental/UAF/UAFAnimGraph/Source/UAFAnimGraph/Public/Traits/BlendTwoWayTraitData.h
new file mode 100644
--- /dev/null
+++ b/Engine/Plugins/Experimental/UAF/UAFAnimGraph/Source/UAFAnimGraph/Public/Traits/BlendTwoWayTraitData.h
@@ -0,0 +1,45 @@
+// Copyright Epic Games, Inc. All Rights Reserved.
+
+#pragma once
+
+#include "TraitCore/TraitBinding.h"
+#include "TraitCore/TraitHandle.h"
+#include "TraitCore/TraitSharedData.h"
+
+#include "BlendTwoWayTraitData.generated.h"
+
+/** A trait that can blend two inputs. */
+USTRUCT(meta = (DisplayName = "Blend Two Way", ShowTooltip=true))
+struct FAnimNextBlendTwoWayTraitSharedData : public FAnimNextTraitSharedData
+{
+	GENERATED_BODY()
+
+	/** First output to be blended (full weight is 0.0). */
+	UPROPERTY()
+	FAnimNextTraitHandle ChildA;
+
+	/** Second output to be blended (full weight is 1.0). */
+	UPROPERTY()
+	FAnimNextTraitHandle ChildB;
+
+	/** How much to blend our two children: 0.0 is fully child A while 1.0 is fully child B. */
+	UPROPERTY(EditAnywhere, Category = "Default")
+	float BlendWeight = 0.0f;
+
+	/** This reinitializes child when re-activated */
+	UPROPERTY(EditAnywhere, Category = "Default")
+	bool bResetChildOnActivation = true;
+
+	/** Always update children, regardless of whether or not that child has weight. */
+	UPROPERTY(EditAnywhere, Category = "Default")
+	bool bAlwaysUpdateChildren = false;
+
+	// Latent pin support boilerplate
+	#define TRAIT_LATENT_PROPERTIES_ENUMERATOR(GeneratorMacro) \
+	GeneratorMacro(BlendWeight) \
+	GeneratorMacro(bResetChildOnActivation) \
+	GeneratorMacro(bAlwaysUpdateChildren) \
+
+	GENERATE_TRAIT_LATENT_PROPERTIES(FAnimNextBlendTwoWayTraitSharedData, TRAIT_LATENT_PROPERTIES_ENUMERATOR)
+	#undef TRAIT_LATENT_PROPERTIES_ENUMERATOR
+};
\ No newline at end of file
diff --git a/Engine/Plugins/Experimental/UAF/UAFAnimGraph/Source/UAFAnimGraph/Public/Traits/PassthroughTraitData.h b/Engine/Plugins/Experimental/UAF/UAFAnimGraph/Source/UAFAnimGraph/Public/Traits/PassthroughTraitData.h
new file mode 100644
--- /dev/null
+++ b/Engine/Plugins/Experimental/UAF/UAFAnimGraph/Source/UAFAnimGraph/Public/Traits/PassthroughTraitData.h
@@ -0,0 +1,19 @@
+// Copyright Epic Games, Inc. All Rights Reserved.
+
+#pragma once
+
+#include "TraitCore/TraitHandle.h"
+#include "TraitCore/TraitSharedData.h"
+
+#include "PassthroughTraitData.generated.h"
+
+/** A trait that passes through the input without modification. */
+USTRUCT(meta = (DisplayName = "Passthrough", ShowTooltip=true))
+struct FAnimNextPassthroughSharedData : public FAnimNextTraitSharedData
+{
+	GENERATED_BODY()
+
+	/** Input to pass to output */
+	UPROPERTY()
+	FAnimNextTraitHandle Input;
+};
\ No newline at end of file
diff --git a/Engine/Plugins/Experimental/UAF/UAFAnimGraph/Source/UAFAnimGraph/Public/Traits/SubGraphHostTraitData.h b/Engine/Plugins/Experimental/UAF/UAFAnimGraph/Source/UAFAnimGraph/Public/Traits/SubGraphHostTraitData.h
new file mode 100644
--- /dev/null
+++ b/Engine/Plugins/Experimental/UAF/UAFAnimGraph/Source/UAFAnimGraph/Public/Traits/SubGraphHostTraitData.h
@@ -0,0 +1,38 @@
+// Copyright Epic Games, Inc. All Rights Reserved.
+
+#pragma once
+
+#include "TraitCore/TraitBinding.h"
+#include "TraitCore/TraitHandle.h"
+#include "TraitCore/TraitSharedData.h"
+
+#include "SubGraphHostTraitData.generated.h"
+
+class UAnimNextAnimationGraph;
+
+/** A trait that hosts and manages a sub-graph instance. */
+USTRUCT(meta = (DisplayName = "SubGraph", ShowTooltip=true))
+struct FAnimNextSubGraphHostTraitSharedData : public FAnimNextTraitSharedData
+{
+	GENERATED_BODY()
+
+	/** Asset to use as a sub-graph */
+	UPROPERTY(EditAnywhere, Category = "Default", meta=(ExportAsReference="true"))
+	TObjectPtr<const UAnimNextAnimationGraph> AnimationGraph;
+
+	/** A dummy child that we can use to output the bind pose. This property is hidden and automatically populated during compilation. */
+	UPROPERTY(meta = (Hidden))
+	FAnimNextTraitHandle ReferencePoseChild;
+
+	/** Entry point in the Subgraph that we will use */
+	UPROPERTY(EditAnywhere, Category = "Default", meta=(CustomWidget = "ParamName", AllowedParamType = "FAnimNextEntryPoint"))
+	FName EntryPoint;
+
+	// Latent pin support boilerplate
+	#define TRAIT_LATENT_PROPERTIES_ENUMERATOR(GeneratorMacro) \
+	GeneratorMacro(AnimationGraph) \
+	GeneratorMacro(EntryPoint) \
+
+	GENERATE_TRAIT_LATENT_PROPERTIES(FAnimNextSubGraphHostTraitSharedData, TRAIT_LATENT_PROPERTIES_ENUMERATOR)
+	#undef TRAIT_LATENT_PROPERTIES_ENUMERATOR
+};
\ No newline at end of file
