UAF: Export ExecutionContext log function and mark it as const

#Synced-CL 49271177

diff --git a/Engine/Plugins/Experimental/UAF/UAFAnimGraph/Source/UAFAnimGraph/Private/TraitCore/ExecutionContext.cpp b/Engine/Plugins/Experimental/UAF/UAFAnimGraph/Source/UAFAnimGraph/Private/TraitCore/ExecutionContext.cpp
--- a/Engine/Plugins/Experimental/UAF/UAFAnimGraph/Source/UAFAnimGraph/Private/TraitCore/ExecutionContext.cpp
+++ b/Engine/Plugins/Experimental/UAF/UAFAnimGraph/Source/UAFAnimGraph/Private/TraitCore/ExecutionContext.cpp
@@ -635,7 +635,7 @@ namespace UE::UAF
 	}
 #endif
 
-	void FExecutionContext::Log(ELogVerbosity::Type Verbosity, const FString& InMessage, const FString& TraitName)
+	void FExecutionContext::Log(ELogVerbosity::Type Verbosity, const FString& InMessage, const FString& TraitName) const
 	{
 		const UObject* HostObject = GetHostObject();
 		const UObject* SystemOwner = HostObject ? HostObject->GetOuter() : nullptr;
diff --git a/Engine/Plugins/Experimental/UAF/UAFAnimGraph/Source/UAFAnimGraph/Public/TraitCore/ExecutionContext.h b/Engine/Plugins/Experimental/UAF/UAFAnimGraph/Source/UAFAnimGraph/Public/TraitCore/ExecutionContext.h
--- a/Engine/Plugins/Experimental/UAF/UAFAnimGraph/Source/UAFAnimGraph/Public/TraitCore/ExecutionContext.h
+++ b/Engine/Plugins/Experimental/UAF/UAFAnimGraph/Source/UAFAnimGraph/Public/TraitCore/ExecutionContext.h
@@ -214,7 +214,7 @@ namespace UE::UAF
 		UE_API FRigVMDrawInterface* GetDebugDrawInterface() const;
 #endif
 		
-		void Log(ELogVerbosity::Type Verbosity, const FString& Message, const FString& TraitName);
+		UE_API void Log(ELogVerbosity::Type Verbosity, const FString& Message, const FString& TraitName) const;
 
 	private:
 		// No copy or move
