UAF: Remove game thread check from UAFMontageComponent destructor and gate unregister call

#Synced-CL 49365099

diff --git a/Engine/Plugins/Experimental/UAF/UAFAnimGraph/Source/UAFAnimGraph/Internal/Injection/UAFMontageComponent.h b/Engine/Plugins/Experimental/UAF/UAFAnimGraph/Source/UAFAnimGraph/Internal/Injection/UAFMontageComponent.h
--- a/Engine/Plugins/Experimental/UAF/UAFAnimGraph/Source/UAFAnimGraph/Internal/Injection/UAFMontageComponent.h
+++ b/Engine/Plugins/Experimental/UAF/UAFAnimGraph/Source/UAFAnimGraph/Internal/Injection/UAFMontageComponent.h
@@ -103,7 +103,9 @@ private:
 	// If the the exposed TickMontagesBeforeEventName does not provide an override, it will fall back to DefaultTickDependencyEventName
 	UPROPERTY(Transient)
 	FName TickDependencyEventName = NAME_None;
-
+	
+	bool bHasTickDependencyRegistered = false;
+	
 	// The event name to fall back to for the module event tick function that the montage tick function will be enqueued as prerequisite for
 	static inline const FLazyName DefaultTickDependencyEventName = FLazyName("PrePhysics");
 
diff --git a/Engine/Plugins/Experimental/UAF/UAFAnimGraph/Source/UAFAnimGraph/Private/Injection/UAFMontageComponent.cpp b/Engine/Plugins/Experimental/UAF/UAFAnimGraph/Source/UAFAnimGraph/Private/Injection/UAFMontageComponent.cpp
--- a/Engine/Plugins/Experimental/UAF/UAFAnimGraph/Source/UAFAnimGraph/Private/Injection/UAFMontageComponent.cpp
+++ b/Engine/Plugins/Experimental/UAF/UAFAnimGraph/Source/UAFAnimGraph/Private/Injection/UAFMontageComponent.cpp
@@ -91,20 +91,20 @@ FName FUAFUpdateMontageTickFunction::DiagnosticContext(bool bDetailed)
 
 FUAFMontageComponent::~FUAFMontageComponent()
 {
-	// This component should only ever get destroyed from the game thread
-	// If this changes in the future we will need to revisit how to safely remove the tick prerequisite
-	check(IsInGameThread());
-
-	if (FAnimNextModuleInstance* ModuleInstance = GetModuleInstancePtr())
+	if (bHasTickDependencyRegistered && IsInGameThread())
 	{
-		if (UAnimNextComponent* UAFComponent = Cast<UAnimNextComponent>(ModuleInstance->GetObject()))
+		if (FAnimNextModuleInstance* ModuleInstance = GetModuleInstancePtr())
 		{
-			for (UE::UAF::FModuleEventTickFunction& TickFunction : ModuleInstance->GetTickFunctions())
+			if (UAnimNextComponent* UAFComponent = Cast<UAnimNextComponent>(ModuleInstance->GetObject()))
 			{
-				if (TickFunction.Event.GetEventName() == TickDependencyEventName)
+				for (UE::UAF::FModuleEventTickFunction& TickFunction : ModuleInstance->GetTickFunctions())
 				{
-					TickFunction.RemovePrerequisite(UAFComponent, MontageTickFunction);
-					break;
+					if (TickFunction.Event.GetEventName() == TickDependencyEventName)
+					{
+						TickFunction.RemovePrerequisite(UAFComponent, MontageTickFunction);
+						bHasTickDependencyRegistered = false;
+						break;
+					}
 				}
 			}
 		}
@@ -212,6 +212,7 @@ void FUAFMontageComponent::OnBindToInstance()
 				MontageComponent->MontageTickFunction.RegisterTickFunction(Level);
 
 				ModuleEventTickFunction->AddPrerequisite(UAFComponent, MontageComponent->MontageTickFunction);
+				MontageComponent->bHasTickDependencyRegistered = true;
 			}
 			else
 			{
