Rename UAF attribute collection and split into 3 types

diff --git a/Engine/Plugins/Experimental/UAF/UAF/Source/UAF/Public/UAF/ValueRuntime/Transformers/Interpolate.h b/Engine/Plugins/Experimental/UAF/UAF/Source/UAF/Public/UAF/ValueRuntime/Transformers/Interpolate.h
--- a/Engine/Plugins/Experimental/UAF/UAF/Source/UAF/Public/UAF/ValueRuntime/Transformers/Interpolate.h
+++ b/Engine/Plugins/Experimental/UAF/UAF/Source/UAF/Public/UAF/ValueRuntime/Transformers/Interpolate.h
@@ -28,13 +28,13 @@ namespace UE::UAF::Transformers
 		// Output = lerp(A, B, WeightB)
 		// Output = (A * (1.0 - WeightB)) + (B * WeightB)
 		// This function supports outputting to either input
-		static UE_API void Apply(const FValueTransformerMapPtr& TransformerMap, const FAttributeCollection& InputA, const FAttributeCollection& InputB, float WeightB, FAttributeCollection& Output);
+		static UE_API void Apply(const FValueTransformerMapPtr& TransformerMap, const FValueBundle& InputA, const FValueBundle& InputB, float WeightB, FValueBundle& Output);
 
 		// Interpolate attribute values
 		// If a weight isn't specified, the default weight will be used
 		// Output = lerp(A, B, HasValueWeightB ? ValueWeightB : DefaultWeightB)
 		// This function supports outputting to either input
-		static UE_API void Apply(const FValueTransformerMapPtr& TransformerMap, const FAttributeCollection& InputA, const FAttributeCollection& InputB, const FAttributeCollection& PerValueWeightsB, float DefaultWeightB, FAttributeCollection& Output);
+		static UE_API void Apply(const FValueTransformerMapPtr& TransformerMap, const FValueBundle& InputA, const FValueBundle& InputB, const FValueBundle& PerValueWeightsB, float DefaultWeightB, FValueBundle& Output);
 	};
 
 	//////////////////////////////////////////////////////////////////////////
diff --git a/Engine/Plugins/Experimental/UAF/UAF/Source/UAF/Public/UAF/ValueRuntime/Transformers/Layer.h b/Engine/Plugins/Experimental/UAF/UAF/Source/UAF/Public/UAF/ValueRuntime/Transformers/Layer.h
--- a/Engine/Plugins/Experimental/UAF/UAF/Source/UAF/Public/UAF/ValueRuntime/Transformers/Layer.h
+++ b/Engine/Plugins/Experimental/UAF/UAF/Source/UAF/Public/UAF/ValueRuntime/Transformers/Layer.h
@@ -30,10 +30,10 @@ namespace UE::UAF::Transformers
 		// This function supports outputting to either input
 		static UE_API void Apply(
 			const FValueTransformerMapPtr& TransformerMap,
-			const FAttributeCollection& Base,
-			const FAttributeCollection& Layer,
+			const FValueBundle& Base,
+			const FValueBundle& Layer,
 			float LayerWeight,
-			FAttributeCollection& Output);
+			FValueBundle& Output);
 
 		// Layer values onto a base
 		// Output = ValueInLayer ? lerp(BaseValue, LayerValue, HasLayerValueWeight ? LayerValueWeight : DefaultLayerWeight) : BaseValue
@@ -41,11 +41,11 @@ namespace UE::UAF::Transformers
 		// This function supports outputting to either input
 		static UE_API void Apply(
 			const FValueTransformerMapPtr& TransformerMap,
-			const FAttributeCollection& Base,
-			const FAttributeCollection& Layer,
-			const FAttributeCollection& PerValueLayerWeights,
+			const FValueBundle& Base,
+			const FValueBundle& Layer,
+			const FValueBundle& PerValueLayerWeights,
 			float DefaultLayerWeight,
-			FAttributeCollection& Output);
+			FValueBundle& Output);
 	};
 
 	//////////////////////////////////////////////////////////////////////////
diff --git a/Engine/Plugins/Experimental/UAF/UAF/Source/UAF/Public/UAF/ValueRuntime/Transformers/Overwrite.h b/Engine/Plugins/Experimental/UAF/UAF/Source/UAF/Public/UAF/ValueRuntime/Transformers/Overwrite.h
--- a/Engine/Plugins/Experimental/UAF/UAF/Source/UAF/Public/UAF/ValueRuntime/Transformers/Overwrite.h
+++ b/Engine/Plugins/Experimental/UAF/UAF/Source/UAF/Public/UAF/ValueRuntime/Transformers/Overwrite.h
@@ -28,13 +28,13 @@ namespace UE::UAF::Transformers
 		// Overwrite attribute values
 		// Output = (A * Weight)
 		// This function supports outputting to its input
-		static UE_API void Apply(const FValueTransformerMapPtr& TransformerMap, const FAttributeCollection& Input, float Weight, FAttributeCollection& Output);
+		static UE_API void Apply(const FValueTransformerMapPtr& TransformerMap, const FValueBundle& Input, float Weight, FValueBundle& Output);
 
 		// Overwrite attribute values using per value weights (must be of type FFloatAnimationAttribute)
 		// If a weight isn't specified, the default weight will be used
 		// Output = A * (HasValueWeight ? ValueWeight : DefaultWeight)
 		// This function supports outputting to its input
-		static UE_API void Apply(const FValueTransformerMapPtr& TransformerMap, const FAttributeCollection& Input, const FAttributeCollection& PerValueWeights, float DefaultWeight, FAttributeCollection& Output);
+		static UE_API void Apply(const FValueTransformerMapPtr& TransformerMap, const FValueBundle& Input, const FValueBundle& PerValueWeights, float DefaultWeight, FValueBundle& Output);
 	};
 
 	//////////////////////////////////////////////////////////////////////////
diff --git a/Engine/Plugins/Experimental/UAF/UAF/Source/UAF/Public/UAF/ValueRuntime/Transformers/Sanitize.h b/Engine/Plugins/Experimental/UAF/UAF/Source/UAF/Public/UAF/ValueRuntime/Transformers/Sanitize.h
--- a/Engine/Plugins/Experimental/UAF/UAF/Source/UAF/Public/UAF/ValueRuntime/Transformers/Sanitize.h
+++ b/Engine/Plugins/Experimental/UAF/UAF/Source/UAF/Public/UAF/ValueRuntime/Transformers/Sanitize.h
@@ -25,7 +25,7 @@ namespace UE::UAF::Transformers
 
 		// Sanitize attribute values
 		// This can be used to normalize rotations (and other like values) as well as check for NaN/Inf values
-		static UE_API void Apply(const FValueTransformerMapPtr& TransformerMap, FAttributeCollection& Collection);
+		static UE_API void Apply(const FValueTransformerMapPtr& TransformerMap, FValueBundle& Collection);
 	};
 
 	//////////////////////////////////////////////////////////////////////////
diff --git a/Engine/Plugins/Experimental/UAF/UAF/Source/UAF/Public/UAF/ValueRuntime/Transformers/Transformer.h b/Engine/Plugins/Experimental/UAF/UAF/Source/UAF/Public/UAF/ValueRuntime/Transformers/Transformer.h
--- a/Engine/Plugins/Experimental/UAF/UAF/Source/UAF/Public/UAF/ValueRuntime/Transformers/Transformer.h
+++ b/Engine/Plugins/Experimental/UAF/UAF/Source/UAF/Public/UAF/ValueRuntime/Transformers/Transformer.h
@@ -3,7 +3,7 @@
 #pragma once
 
 #include "CoreMinimal.h"
-#include "UAF/ValueRuntime/AttributeCollection.h"
+#include "UAF/ValueRuntime/ValueBundle.h"
 
 #define UE_API UAF_API
 
diff --git a/Engine/Plugins/Experimental/UAF/UAF/Source/UAF/Public/UAF/ValueRuntime/UnboundValueMap.h b/Engine/Plugins/Experimental/UAF/UAF/Source/UAF/Public/UAF/ValueRuntime/UnboundValueMap.h
--- a/Engine/Plugins/Experimental/UAF/UAF/Source/UAF/Public/UAF/ValueRuntime/UnboundValueMap.h
+++ b/Engine/Plugins/Experimental/UAF/UAF/Source/UAF/Public/UAF/ValueRuntime/UnboundValueMap.h
@@ -100,7 +100,7 @@ namespace UE::UAF
 		// The size in bytes of this container (excluding any externally allocated data)
 		int32 ContainerSize = 0;
 
-		friend class FAttributeCollection;
+		friend class FValueBundle;
 		friend class FValueRuntimeRegistry;
 		friend void ReleaseUnboundValueMap(FUnboundValueMap* Map);
 	};
diff --git a/Engine/Plugins/Experimental/UAF/UAFAnimGraph/Source/UAFAnimGraph/Private/EvaluationVM/Tasks/ApplyAdditiveKeyframe.cpp b/Engine/Plugins/Experimental/UAF/UAFAnimGraph/Source/UAFAnimGraph/Private/EvaluationVM/Tasks/ApplyAdditiveKeyframe.cpp
--- a/Engine/Plugins/Experimental/UAF/UAFAnimGraph/Source/UAFAnimGraph/Private/EvaluationVM/Tasks/ApplyAdditiveKeyframe.cpp
+++ b/Engine/Plugins/Experimental/UAF/UAFAnimGraph/Source/UAFAnimGraph/Private/EvaluationVM/Tasks/ApplyAdditiveKeyframe.cpp
@@ -5,7 +5,7 @@
 #include "EvaluationVM/EvaluationVM.h"
 #include "EvaluationVM/KeyframeState.h"
 #include "TransformArrayOperations.h"
-#include "UAF/ValueRuntime/AttributeCollection.h"
+#include "UAF/ValueRuntime/ValueBundle.h"
 #include "UAF/ValueRuntime/Transformers/AdditiveSpace.h"
 #include "UAF/ValueRuntime/Transformers/BoneSpace.h"
 
@@ -34,27 +34,27 @@ void FAnimNextApplyAdditiveKeyframeTask::Execute(UE::UAF::FEvaluationVM& VM) con
 
 	if (VM.GetActiveNamedSet())
 	{
-		TUniquePtr<FAttributeCollection> AdditiveCollection;
+		TUniquePtr<FValueBundle> AdditiveCollection;
 		if (!VM.PopValue(ATTRIBUTE_STACK_NAME, AdditiveCollection))
 		{
 			// We have no inputs, nothing to do
-			FAttributeCollectionStack Collection(VM.GetActiveNamedSet());
+			FValueBundleStack Collection(VM.GetActiveNamedSet());
 			Collection.InitWithValueSpace(FValueSpace(EValueSpaceType::Local));
 
-			VM.PushValue(ATTRIBUTE_STACK_NAME, MakeUnique<FAttributeCollection>(MoveTemp(Collection)));
+			VM.PushValue(ATTRIBUTE_STACK_NAME, MakeUnique<FValueBundle>(MoveTemp(Collection)));
 			UE_LOG(LogAnimation, Warning, TEXT("FAnimNextApplyAdditiveKeyframeTask::Execute: Could not apply additive keyframe - no inputs provided. Pushing a ref pose instead."));
 			return;
 		}
 
-		TUniquePtr<FAttributeCollection> BaseCollection;
+		TUniquePtr<FValueBundle> BaseCollection;
 		if (!VM.PopValue(ATTRIBUTE_STACK_NAME, BaseCollection))
 		{
 			// We have a single input, discard it since it must be the additive pose, either way something went wrong
 			// Push the reference pose since we'll expect a non-additive pose
-			FAttributeCollectionStack Collection(VM.GetActiveNamedSet());
+			FValueBundleStack Collection(VM.GetActiveNamedSet());
 			Collection.InitWithValueSpace(FValueSpace(EValueSpaceType::Local));
 
-			VM.PushValue(ATTRIBUTE_STACK_NAME, MakeUnique<FAttributeCollection>(MoveTemp(Collection)));
+			VM.PushValue(ATTRIBUTE_STACK_NAME, MakeUnique<FValueBundle>(MoveTemp(Collection)));
 			UE_LOG(LogAnimation, Warning, TEXT("FAnimNextApplyAdditiveKeyframeTask::Execute: Could not apply additive keyframe - only a single input was provided. Pushing a ref pose instead."));
 			return;
 		}
@@ -63,10 +63,10 @@ void FAnimNextApplyAdditiveKeyframeTask::Execute(UE::UAF::FEvaluationVM& VM) con
 		if (!AdditiveSpace.IsAdditive())
 		{
 			// Additive input must be additive, push reference pose if not the case
-			FAttributeCollectionStack Collection(VM.GetActiveNamedSet());
+			FValueBundleStack Collection(VM.GetActiveNamedSet());
 			Collection.InitWithValueSpace(FValueSpace(EValueSpaceType::Local));
 
-			VM.PushValue(ATTRIBUTE_STACK_NAME, MakeUnique<FAttributeCollection>(MoveTemp(Collection)));
+			VM.PushValue(ATTRIBUTE_STACK_NAME, MakeUnique<FValueBundle>(MoveTemp(Collection)));
 			UE_LOG(LogAnimation, Warning, TEXT("FAnimNextApplyAdditiveKeyframeTask::Execute: Could not apply additive keyframe - The expected additive pose was not of an additive type. Pushing a ref pose instead."));
 			return;
 		}
@@ -184,7 +184,7 @@ float FAnimNextApplyAdditiveKeyframeTask::GetInterpolationAlpha(const UE::UAF::F
 	return FMath::Clamp(Alpha, 0.0f, 1.0f);
 }
 
-float FAnimNextApplyAdditiveKeyframeTask::GetInterpolationAlpha(const UE::UAF::FAttributeCollection& CollectionA, const UE::UAF::FAttributeCollection& CollectionB) const
+float FAnimNextApplyAdditiveKeyframeTask::GetInterpolationAlpha(const UE::UAF::FValueBundle& CollectionA, const UE::UAF::FValueBundle& CollectionB) const
 {
 	using namespace UE::UAF;
 
@@ -193,8 +193,8 @@ float FAnimNextApplyAdditiveKeyframeTask::GetInterpolationAlpha(const UE::UAF::F
 	if (AlphaSourceCurveName != NAME_None && AlphaCurveInputIndex != INDEX_NONE)
 	{
 		const FAttributeMappingKey MappingKey = FAttributeMappingKey::MakeFromTo<FFloatAnimationAttribute>();
-		const FAttributeCollection& AlphaCurveCollection = AlphaCurveInputIndex == 0 ? CollectionA : CollectionB;
-		if (const TBoundValueMap<FFloatAnimationAttribute>* FloatMap = AlphaCurveCollection.FindSetMap<FFloatAnimationAttribute>(MappingKey))
+		const FValueBundle& AlphaCurveCollection = AlphaCurveInputIndex == 0 ? CollectionA : CollectionB;
+		if (const TBoundValueMap<FFloatAnimationAttribute>* FloatMap = AlphaCurveCollection.GetBoundValueMaps().Find<FFloatAnimationAttribute>(MappingKey))
 		{
 			const FAttributeTypedSetPtr& TypedSet = FloatMap->GetTypedSet();
 			if (FAttributeSetIndex CurveIndex = TypedSet->FindIndex(AlphaSourceCurveName))
diff --git a/Engine/Plugins/Experimental/UAF/UAFAnimGraph/Source/UAFAnimGraph/Private/EvaluationVM/Tasks/ApplyAdditiveKeyframePerBone.cpp b/Engine/Plugins/Experimental/UAF/UAFAnimGraph/Source/UAFAnimGraph/Private/EvaluationVM/Tasks/ApplyAdditiveKeyframePerBone.cpp
--- a/Engine/Plugins/Experimental/UAF/UAFAnimGraph/Source/UAFAnimGraph/Private/EvaluationVM/Tasks/ApplyAdditiveKeyframePerBone.cpp
+++ b/Engine/Plugins/Experimental/UAF/UAFAnimGraph/Source/UAFAnimGraph/Private/EvaluationVM/Tasks/ApplyAdditiveKeyframePerBone.cpp
@@ -6,7 +6,7 @@
 #include "EvaluationVM/KeyframeState.h"
 #include "TransformArrayOperations.h"
 #include "Animation/AnimationAsset.h"
-#include "UAF/ValueRuntime/AttributeCollection.h"
+#include "UAF/ValueRuntime/ValueBundle.h"
 #include "UAF/ValueRuntime/Transformers/AdditiveSpace.h"
 #include "UAF/ValueRuntime/Transformers/BoneSpace.h"
 
@@ -46,22 +46,22 @@ void FUAFApplyAdditiveKeyframePerBoneTask::Execute(UE::UAF::FEvaluationVM& VM) c
 	if (VM.GetActiveNamedSet())
 	{
 		// Pop our top two poses, the top one being the additive pose and the second one the base
-		TUniquePtr<FAttributeCollection> AdditiveInput;
+		TUniquePtr<FValueBundle> AdditiveInput;
 		if (!VM.PopValue(ATTRIBUTE_STACK_NAME, AdditiveInput))
 		{
 			// We have no inputs, nothing to do
 			return;
 		}
 
-		TUniquePtr<FAttributeCollection> BaseInput;
+		TUniquePtr<FValueBundle> BaseInput;
 		if (!VM.PopValue(ATTRIBUTE_STACK_NAME, BaseInput))
 		{
 			// We have a single input, discard it since it must be the additive pose, either way something went wrong
 			// Push the reference pose since we'll expect a non-additive pose
-			FAttributeCollectionStack Collection(VM.GetActiveNamedSet());
+			FValueBundleStack Collection(VM.GetActiveNamedSet());
 			Collection.InitWithValueSpace(FValueSpace(EValueSpaceType::Local));
 
-			VM.PushValue(ATTRIBUTE_STACK_NAME, MakeUnique<FAttributeCollection>(MoveTemp(Collection)));
+			VM.PushValue(ATTRIBUTE_STACK_NAME, MakeUnique<FValueBundle>(MoveTemp(Collection)));
 			UE_LOG(LogAnimation, Warning, TEXT("FUAFApplyAdditiveKeyframePerBoneTask::Execute: Could not apply additive keyframe - only a single input was provided. Pushing a ref pose instead."));
 			return;
 		}
@@ -70,10 +70,10 @@ void FUAFApplyAdditiveKeyframePerBoneTask::Execute(UE::UAF::FEvaluationVM& VM) c
 		if (!AdditiveSpace.IsAdditive())
 		{
 			// Additive input must be additive, push reference pose if not the case
-			FAttributeCollectionStack Collection(VM.GetActiveNamedSet());
+			FValueBundleStack Collection(VM.GetActiveNamedSet());
 			Collection.InitWithValueSpace(FValueSpace(EValueSpaceType::Local));
 
-			VM.PushValue(ATTRIBUTE_STACK_NAME, MakeUnique<FAttributeCollection>(MoveTemp(Collection)));
+			VM.PushValue(ATTRIBUTE_STACK_NAME, MakeUnique<FValueBundle>(MoveTemp(Collection)));
 			UE_LOG(LogAnimation, Warning, TEXT("FUAFApplyAdditiveKeyframePerBoneTask::Execute: Could not apply additive keyframe - The expected additive pose was not of an additive type. Pushing a ref pose instead."));
 			return;
 		}
@@ -83,16 +83,17 @@ void FUAFApplyAdditiveKeyframePerBoneTask::Execute(UE::UAF::FEvaluationVM& VM) c
 		UScriptStruct* FloatAttributeType = FFloatAnimationAttribute::StaticStruct();
 
 		// Build our per-value weights (only for bound values)
-		FAttributeCollectionStack PerValueWeights(VM.GetActiveNamedSet());
+		FValueBundleStack PerValueWeights(VM.GetActiveNamedSet());
+		FBoundMapCollection& PerValueWeightMaps = PerValueWeights.GetBoundValueMaps();
 
-		for (auto It = BaseInput->CreateConstSetMapIterator(); It; ++It)
+		for (auto It = BaseInput->GetBoundValueMaps().CreateConstIterator(); It; ++It)
 		{
 			UScriptStruct* AttributeType = It.GetAttributeType();
 
 			// Every attribute type maps to a float, our weight
 			FAttributeMappingKey MappingKey = FAttributeMappingKey::MakeFromTo(AttributeType, FloatAttributeType);
 
-			TBoundValueMap<FFloatAnimationAttribute>* SetWeights = PerValueWeights.AddSetMap<FFloatAnimationAttribute>(MappingKey);
+			TBoundValueMap<FFloatAnimationAttribute>* SetWeights = PerValueWeightMaps.Add<FFloatAnimationAttribute>(MappingKey);
 
 			if (const TBoundValueMap<FBoneTransformAnimationAttribute>* BoneMap = Cast<FBoneTransformAnimationAttribute>(It.GetMap()))
 			{
diff --git a/Engine/Plugins/Experimental/UAF/UAFAnimGraph/Source/UAFAnimGraph/Private/EvaluationVM/Tasks/BlendKeyframes.cpp b/Engine/Plugins/Experimental/UAF/UAFAnimGraph/Source/UAFAnimGraph/Private/EvaluationVM/Tasks/BlendKeyframes.cpp
--- a/Engine/Plugins/Experimental/UAF/UAFAnimGraph/Source/UAFAnimGraph/Private/EvaluationVM/Tasks/BlendKeyframes.cpp
+++ b/Engine/Plugins/Experimental/UAF/UAFAnimGraph/Source/UAFAnimGraph/Private/EvaluationVM/Tasks/BlendKeyframes.cpp
@@ -6,7 +6,7 @@
 #include "EvaluationVM/KeyframeState.h"
 #include "TransformArrayOperations.h"
 #include "Animation/InputScaleBias.h"
-#include "UAF/ValueRuntime/AttributeCollection.h"
+#include "UAF/ValueRuntime/ValueBundle.h"
 #include "UAF/ValueRuntime/Transformers/Accumulate.h"
 #include "UAF/ValueRuntime/Transformers/Interpolate.h"
 #include "UAF/ValueRuntime/Transformers/Overwrite.h"
@@ -39,14 +39,14 @@ void FAnimNextBlendTwoKeyframesTask::Execute(UE::UAF::FEvaluationVM& VM) const
 	{
 		// Pop our top two poses, we'll re-use the top keyframe for our result
 
-		TUniquePtr<FAttributeCollection> CollectionB;
+		TUniquePtr<FValueBundle> CollectionB;
 		if (!VM.PopValue(ATTRIBUTE_STACK_NAME, CollectionB))
 		{
 			// We have no inputs, nothing to do
 			return;
 		}
 
-		TUniquePtr<FAttributeCollection> CollectionA;
+		TUniquePtr<FValueBundle> CollectionA;
 		if (!VM.PopValue(ATTRIBUTE_STACK_NAME, CollectionA))
 		{
 			// We have a single input, leave it on top of the stack
@@ -164,7 +164,7 @@ float FAnimNextBlendTwoKeyframesTask::GetInterpolationAlpha(const UE::UAF::FKeyf
 	return FMath::Clamp(Alpha, 0.0f, 1.0f);
 }
 
-float FAnimNextBlendTwoKeyframesTask::GetInterpolationAlpha(const UE::UAF::FAttributeCollection& CollectionA, const UE::UAF::FAttributeCollection& CollectionB) const
+float FAnimNextBlendTwoKeyframesTask::GetInterpolationAlpha(const UE::UAF::FValueBundle& CollectionA, const UE::UAF::FValueBundle& CollectionB) const
 {
 	using namespace UE::UAF;
 
@@ -173,8 +173,8 @@ float FAnimNextBlendTwoKeyframesTask::GetInterpolationAlpha(const UE::UAF::FAttr
 	if (AlphaSourceCurveName != NAME_None && AlphaCurveInputIndex != INDEX_NONE)
 	{
 		const FAttributeMappingKey MappingKey = FAttributeMappingKey::MakeFromTo<FFloatAnimationAttribute>();
-		const FAttributeCollection& AlphaCurveCollection = AlphaCurveInputIndex == 0 ? CollectionA : CollectionB;
-		if (const TBoundValueMap<FFloatAnimationAttribute>* FloatMap = AlphaCurveCollection.FindSetMap<FFloatAnimationAttribute>(MappingKey))
+		const FValueBundle& AlphaCurveCollection = AlphaCurveInputIndex == 0 ? CollectionA : CollectionB;
+		if (const TBoundValueMap<FFloatAnimationAttribute>* FloatMap = AlphaCurveCollection.GetBoundValueMaps().Find<FFloatAnimationAttribute>(MappingKey))
 		{
 			const FAttributeTypedSetPtr& TypedSet = FloatMap->GetTypedSet();
 			if (FAttributeSetIndex CurveIndex = TypedSet->FindIndex(AlphaSourceCurveName))
@@ -218,7 +218,7 @@ void FAnimNextBlendOverwriteKeyframeWithScaleTask::Execute(UE::UAF::FEvaluationV
 
 	if (VM.GetActiveNamedSet())
 	{
-		TUniquePtr<FAttributeCollection> Collection;
+		TUniquePtr<FValueBundle> Collection;
 		if (!VM.PopValue(ATTRIBUTE_STACK_NAME, Collection))
 		{
 			// We have no inputs, nothing to do
@@ -281,14 +281,14 @@ void FAnimNextBlendAddKeyframeWithScaleTask::Execute(UE::UAF::FEvaluationVM& VM)
 	{
 		// Pop our top two poses, we'll re-use the top keyframe for our result
 
-		TUniquePtr<FAttributeCollection> CollectionB;
+		TUniquePtr<FValueBundle> CollectionB;
 		if (!VM.PopValue(ATTRIBUTE_STACK_NAME, CollectionB))
 		{
 			// We have no inputs, nothing to do
 			return;
 		}
 
-		TUniquePtr<FAttributeCollection> CollectionA;
+		TUniquePtr<FValueBundle> CollectionA;
 		if (!VM.PopValue(ATTRIBUTE_STACK_NAME, CollectionA))
 		{
 			// We have a single input, leave it on top of the stack
diff --git a/Engine/Plugins/Experimental/UAF/UAFAnimGraph/Source/UAFAnimGraph/Private/EvaluationVM/Tasks/BlendKeyframesPerBone.cpp b/Engine/Plugins/Experimental/UAF/UAFAnimGraph/Source/UAFAnimGraph/Private/EvaluationVM/Tasks/BlendKeyframesPerBone.cpp
--- a/Engine/Plugins/Experimental/UAF/UAFAnimGraph/Source/UAFAnimGraph/Private/EvaluationVM/Tasks/BlendKeyframesPerBone.cpp
+++ b/Engine/Plugins/Experimental/UAF/UAFAnimGraph/Source/UAFAnimGraph/Private/EvaluationVM/Tasks/BlendKeyframesPerBone.cpp
@@ -20,8 +20,8 @@
 
 namespace UE::UAF::Private
 {
-	static FAttributeCollectionStack BuildPerValueWeights(
-		const FAttributeCollection& Input,
+	static FValueBundleStack BuildPerValueWeights(
+		const FValueBundle& Input,
 		UE::UAF::FEvaluationVM& VM,
 		const FBlendSampleData* BlendData,
 		const IBlendProfileInterface* BlendProfile,
@@ -32,16 +32,17 @@ namespace UE::UAF::Private
 		UScriptStruct* FloatAttributeType = FFloatAnimationAttribute::StaticStruct();
 
 		// Build our per-value weights (only for bound values)
-		FAttributeCollectionStack PerValueWeights(VM.GetActiveNamedSet());
+		FValueBundleStack PerValueWeights(VM.GetActiveNamedSet());
+		FBoundMapCollection& PerValueWeightMaps = PerValueWeights.GetBoundValueMaps();
 
-		for (auto It = Input.CreateConstSetMapIterator(); It; ++It)
+		for (auto It = Input.GetBoundValueMaps().CreateConstIterator(); It; ++It)
 		{
 			UScriptStruct* AttributeType = It.GetAttributeType();
 
 			// Every attribute type maps to a float, our weight
 			FAttributeMappingKey MappingKey = FAttributeMappingKey::MakeFromTo(AttributeType, FloatAttributeType);
 
-			TBoundValueMap<FFloatAnimationAttribute>* SetWeights = PerValueWeights.AddSetMap<FFloatAnimationAttribute>(MappingKey);
+			TBoundValueMap<FFloatAnimationAttribute>* SetWeights = PerValueWeightMaps.Add<FFloatAnimationAttribute>(MappingKey);
 
 			if (const TBoundValueMap<FBoneTransformAnimationAttribute>* BoneMap = Cast<FBoneTransformAnimationAttribute>(It.GetMap()))
 			{
@@ -113,8 +114,8 @@ namespace UE::UAF::Private
 		return PerValueWeights;
 	}
 
-	static FAttributeCollectionStack BuildPerValueWeights(
-		const FAttributeCollection& Input,
+	static FValueBundleStack BuildPerValueWeights(
+		const FValueBundle& Input,
 		UE::UAF::FEvaluationVM& VM,
 		const USkeleton* SourceSkeleton,
 		const TArray<float>& BoneMaskWeights,
@@ -127,16 +128,17 @@ namespace UE::UAF::Private
 
 		// Build our per-value weights (only for bound values)
 		// By default, the layered values are ignored unless they have weight
-		FAttributeCollectionStack PerValueWeights(VM.GetActiveNamedSet());
+		FValueBundleStack PerValueWeights(VM.GetActiveNamedSet());
+		FBoundMapCollection& PerValueWeightMaps = PerValueWeights.GetBoundValueMaps();
 
-		for (auto It = Input.CreateConstSetMapIterator(); It; ++It)
+		for (auto It = Input.GetBoundValueMaps().CreateConstIterator(); It; ++It)
 		{
 			UScriptStruct* AttributeType = It.GetAttributeType();
 
 			// Every attribute type maps to a float, our weight
 			FAttributeMappingKey MappingKey = FAttributeMappingKey::MakeFromTo(AttributeType, FloatAttributeType);
 
-			TBoundValueMap<FFloatAnimationAttribute>* SetWeights = PerValueWeights.AddSetMap<FFloatAnimationAttribute>(MappingKey);
+			TBoundValueMap<FFloatAnimationAttribute>* SetWeights = PerValueWeightMaps.Add<FFloatAnimationAttribute>(MappingKey);
 
 			if (const TBoundValueMap<FBoneTransformAnimationAttribute>* BoneMap = Cast<FBoneTransformAnimationAttribute>(It.GetMap()))
 			{
@@ -255,7 +257,7 @@ void FAnimNextBlendOverwriteKeyframePerBoneWithScaleTask::Execute(UE::UAF::FEval
 
 	if (VM.GetActiveNamedSet())
 	{
-		TUniquePtr<FAttributeCollection> Input;
+		TUniquePtr<FValueBundle> Input;
 		if (!VM.PopValue(ATTRIBUTE_STACK_NAME, Input))
 		{
 			// We have no inputs, nothing to do
@@ -264,7 +266,7 @@ void FAnimNextBlendOverwriteKeyframePerBoneWithScaleTask::Execute(UE::UAF::FEval
 
 		const float BlendWeight = BlendData->GetClampedWeight();
 
-		const FAttributeCollectionStack PerValueWeights = Private::BuildPerValueWeights(*Input, VM, BlendData, BlendProfile, bPerBoneDataUsesSkeletonIndex);
+		const FValueBundleStack PerValueWeights = Private::BuildPerValueWeights(*Input, VM, BlendData, BlendProfile, bPerBoneDataUsesSkeletonIndex);
 
 		// Modify our input in place
 		Transformers::FOverwrite::Apply(VM.GetTransformerMap(), *Input, PerValueWeights, BlendWeight, *Input);
@@ -397,14 +399,14 @@ void FAnimNextBlendAddKeyframePerBoneWithScaleTask::Execute(UE::UAF::FEvaluation
 
 	if (VM.GetActiveNamedSet())
 	{
-		TUniquePtr<FAttributeCollection> InputB;
+		TUniquePtr<FValueBundle> InputB;
 		if (!VM.PopValue(ATTRIBUTE_STACK_NAME, InputB))
 		{
 			// We have no inputs, nothing to do
 			return;
 		}
 
-		TUniquePtr<FAttributeCollection> InputA;
+		TUniquePtr<FValueBundle> InputA;
 		if (!VM.PopValue(ATTRIBUTE_STACK_NAME, InputA))
 		{
 			// We have a single input, leave it on top of the stack
@@ -414,7 +416,7 @@ void FAnimNextBlendAddKeyframePerBoneWithScaleTask::Execute(UE::UAF::FEvaluation
 
 		const float BlendWeight = BlendDataA->GetClampedWeight();
 
-		const FAttributeCollectionStack PerValueWeightsA = Private::BuildPerValueWeights(*InputA, VM, BlendDataA, BlendProfile, bPerBoneDataUsesSkeletonIndex);
+		const FValueBundleStack PerValueWeightsA = Private::BuildPerValueWeights(*InputA, VM, BlendDataA, BlendProfile, bPerBoneDataUsesSkeletonIndex);
 
 		// Modify input B in place
 		// We flip A and B here because state is pushed in reversed order
@@ -551,14 +553,14 @@ void FAnimNextBlendKeyframePerBoneWithScaleTask::Execute(UE::UAF::FEvaluationVM&
 
 	if (VM.GetActiveNamedSet())
 	{
-		TUniquePtr<FAttributeCollection> Layer;
+		TUniquePtr<FValueBundle> Layer;
 		if (!VM.PopValue(ATTRIBUTE_STACK_NAME, Layer))
 		{
 			// We have no inputs, nothing to do
 			return;
 		}
 
-		TUniquePtr<FAttributeCollection> Base;
+		TUniquePtr<FValueBundle> Base;
 		if (!VM.PopValue(ATTRIBUTE_STACK_NAME, Base))
 		{
 			// We have a single input, leave it on top of the stack
@@ -566,7 +568,7 @@ void FAnimNextBlendKeyframePerBoneWithScaleTask::Execute(UE::UAF::FEvaluationVM&
 			return;
 		}
 
-		const FAttributeCollectionStack PerValueWeightsLayer = Private::BuildPerValueWeights(*Base, VM, SourceSkeleton, BoneMaskWeights, *BlendProfile, CurveMaskWeights, AttributeMaskWeights, ScaleFactor);
+		const FValueBundleStack PerValueWeightsLayer = Private::BuildPerValueWeights(*Base, VM, SourceSkeleton, BoneMaskWeights, *BlendProfile, CurveMaskWeights, AttributeMaskWeights, ScaleFactor);
 
 		// Apply our layer onto our base and modify it in-place
 		Transformers::FLayer::Apply(VM.GetTransformerMap(), *Base, *Layer, PerValueWeightsLayer, ScaleFactor, *Base);
diff --git a/Engine/Plugins/Experimental/UAF/UAFAnimGraph/Source/UAFAnimGraph/Private/EvaluationVM/Tasks/ConvertRotationsToLocalSpace.cpp b/Engine/Plugins/Experimental/UAF/UAFAnimGraph/Source/UAFAnimGraph/Private/EvaluationVM/Tasks/ConvertRotationsToLocalSpace.cpp
--- a/Engine/Plugins/Experimental/UAF/UAFAnimGraph/Source/UAFAnimGraph/Private/EvaluationVM/Tasks/ConvertRotationsToLocalSpace.cpp
+++ b/Engine/Plugins/Experimental/UAF/UAFAnimGraph/Source/UAFAnimGraph/Private/EvaluationVM/Tasks/ConvertRotationsToLocalSpace.cpp
@@ -28,7 +28,7 @@ void FAnimNextConvertRotationsToLocalSpaceTask::Execute(UE::UAF::FEvaluationVM&
 	{
 		for (int32 PoseIndex = 0; PoseIndex < NumPoses; ++PoseIndex)
 		{
-			TUniquePtr<FAttributeCollection>* Collection = VM.PeekValueMutable<TUniquePtr<FAttributeCollection>>(ATTRIBUTE_STACK_NAME, PoseIndex);
+			TUniquePtr<FValueBundle>* Collection = VM.PeekValueMutable<TUniquePtr<FValueBundle>>(ATTRIBUTE_STACK_NAME, PoseIndex);
 			if (Collection == nullptr || !Collection->IsValid())
 			{
 				// No pose to convert, we are done
diff --git a/Engine/Plugins/Experimental/UAF/UAFAnimGraph/Source/UAFAnimGraph/Private/EvaluationVM/Tasks/ConvertRotationsToMeshSpace.cpp b/Engine/Plugins/Experimental/UAF/UAFAnimGraph/Source/UAFAnimGraph/Private/EvaluationVM/Tasks/ConvertRotationsToMeshSpace.cpp
--- a/Engine/Plugins/Experimental/UAF/UAFAnimGraph/Source/UAFAnimGraph/Private/EvaluationVM/Tasks/ConvertRotationsToMeshSpace.cpp
+++ b/Engine/Plugins/Experimental/UAF/UAFAnimGraph/Source/UAFAnimGraph/Private/EvaluationVM/Tasks/ConvertRotationsToMeshSpace.cpp
@@ -28,7 +28,7 @@ void FAnimNextConvertRotationsToMeshSpaceTask::Execute(UE::UAF::FEvaluationVM& V
 	{
 		for (int32 PoseIndex = 0; PoseIndex < NumPoses; ++PoseIndex)
 		{
-			TUniquePtr<FAttributeCollection>* Collection = VM.PeekValueMutable<TUniquePtr<FAttributeCollection>>(ATTRIBUTE_STACK_NAME, PoseIndex);
+			TUniquePtr<FValueBundle>* Collection = VM.PeekValueMutable<TUniquePtr<FValueBundle>>(ATTRIBUTE_STACK_NAME, PoseIndex);
 			if (Collection == nullptr || !Collection->IsValid())
 			{
 				// No pose to convert, we are done
diff --git a/Engine/Plugins/Experimental/UAF/UAFAnimGraph/Source/UAFAnimGraph/Private/EvaluationVM/Tasks/DecompressionTools.cpp b/Engine/Plugins/Experimental/UAF/UAFAnimGraph/Source/UAFAnimGraph/Private/EvaluationVM/Tasks/DecompressionTools.cpp
--- a/Engine/Plugins/Experimental/UAF/UAFAnimGraph/Source/UAFAnimGraph/Private/EvaluationVM/Tasks/DecompressionTools.cpp
+++ b/Engine/Plugins/Experimental/UAF/UAFAnimGraph/Source/UAFAnimGraph/Private/EvaluationVM/Tasks/DecompressionTools.cpp
@@ -16,7 +16,7 @@
 #include "Animation/AnimSequence.h"
 #include "HAL/ConsoleManager.h"
 #include "Logging/StructuredLog.h"
-#include "UAF/ValueRuntime/AttributeCollection.h"
+#include "UAF/ValueRuntime/ValueBundle.h"
 #include "EvaluationVM/EvaluationTaskContext.h"
 #include "UAF/Attributes/EngineAttributes.h"
 #include "RetargetingTools.h"
@@ -294,7 +294,7 @@ namespace Private
 				// and we have nothing to subtract it from, we should probably use the additive identity
 				const FAttributeNamedSetPtr& NamedSet = OutBoneTransforms.GetTypedSet()->GetNamedSet();
 				const FAttributeCollectionPtr ReferenceValues = NamedSet->GetDefaultAttributeValues();
-				const TBoundValueMap<FBoneTransformAnimationAttribute>* ReferenceBoneTransforms = ReferenceValues->FindSetMap<FBoneTransformAnimationAttribute>(OutBoneTransforms.GetMappingKey());
+				const TBoundValueMap<FBoneTransformAnimationAttribute>* ReferenceBoneTransforms = ReferenceValues->GetBoundValueMaps().Find<FBoneTransformAnimationAttribute>(OutBoneTransforms.GetMappingKey());
 				check(ReferenceBoneTransforms != nullptr && !ReferenceBoneTransforms->IsEmpty());
 
 				const FTransform& RootRefTransform = ReferenceBoneTransforms->GetData()->Value;
@@ -604,12 +604,12 @@ void FDecompressionTools::GetAnimationPose(const UAnimSequence* AnimSequence, co
 
 }
 
-void FDecompressionTools::GetAnimationPose(const UAnimSequence* AnimSequence, const FAnimExtractContext& ExtractionContext, const FEvaluationTaskContext& VMContext, FAttributeCollection& OutCollection, bool bForceUseRawData)
+void FDecompressionTools::GetAnimationPose(const UAnimSequence* AnimSequence, const FAnimExtractContext& ExtractionContext, const FEvaluationTaskContext& VMContext, FValueBundle& OutCollection, bool bForceUseRawData)
 {
 	using namespace UE::UAF;
 
 	const FAttributeMappingKey MappingKey = FAttributeMappingKey::MakeFromTo<FBoneTransformAnimationAttribute>();
-	TBoundValueMap<FBoneTransformAnimationAttribute>* BoneTransforms = OutCollection.FindSetMap<FBoneTransformAnimationAttribute>(MappingKey);
+	TBoundValueMap<FBoneTransformAnimationAttribute>* BoneTransforms = OutCollection.GetBoundValueMaps().Find<FBoneTransformAnimationAttribute>(MappingKey);
 	if (BoneTransforms == nullptr || BoneTransforms->IsEmpty())
 	{
 		// Not requesting bone transforms
@@ -777,7 +777,7 @@ void FDecompressionTools::GetAnimationCurves(const UAnimSequence* AnimSequence,
 	AnimSequence->EvaluateCurveData(OutCurves, ExtractionContext.CurrentTime, bForceUseRawData);
 }
 
-void FDecompressionTools::GetAnimationCurves(const UAnimSequence* AnimSequence, const FAnimExtractContext& ExtractionContext, const FEvaluationTaskContext& VMContext, FAttributeCollection& OutCollection, bool bForceUseRawData)
+void FDecompressionTools::GetAnimationCurves(const UAnimSequence* AnimSequence, const FAnimExtractContext& ExtractionContext, const FEvaluationTaskContext& VMContext, FValueBundle& OutCollection, bool bForceUseRawData)
 {
 	using namespace UE::UAF;
 
@@ -817,7 +817,7 @@ void FDecompressionTools::GetAnimationCurves(const UAnimSequence* AnimSequence,
 				ensureMsgf(Curves.IsSorted(), TEXT("Decompressed curves must be sorted"));
 
 				const FAttributeMappingKey MappingKey = FAttributeMappingKey::MakeFromTo<FFloatAnimationAttribute>();
-				TBoundValueMap<FFloatAnimationAttribute>* FloatAttributes = OutCollection.FindSetMap<FFloatAnimationAttribute>(MappingKey);
+				TBoundValueMap<FFloatAnimationAttribute>* FloatAttributes = OutCollection.GetBoundValueMaps().Find<FFloatAnimationAttribute>(MappingKey);
 				TUnboundValueMap<FFloatAnimationAttribute>* FloatMap = nullptr;
 
 				if (FloatAttributes != nullptr && !FloatAttributes->IsEmpty())
@@ -849,7 +849,7 @@ void FDecompressionTools::GetAnimationCurves(const UAnimSequence* AnimSequence,
 							// This curve isn't mapped within our binding, it is dynamically defined
 							if (FloatMap == nullptr)
 							{
-								FloatMap = OutCollection.AddMap<FFloatAnimationAttribute>();
+								FloatMap = OutCollection.GetUnboundValueMaps().Add<FFloatAnimationAttribute>();
 							}
 
 							FloatMap->Append(CurveElement.Name, FFloatAnimationAttribute{ CurveElement.Value });
@@ -859,7 +859,7 @@ void FDecompressionTools::GetAnimationCurves(const UAnimSequence* AnimSequence,
 				else
 				{
 					// We have no float attributes within our named set, everything is dynamically defined
-					FloatMap = OutCollection.AddMap<FFloatAnimationAttribute>();
+					FloatMap = OutCollection.GetUnboundValueMaps().Add<FFloatAnimationAttribute>();
 
 					for (int32 DecompressedCurveIndex = 0; DecompressedCurveIndex < NumDecompressedCurves; ++DecompressedCurveIndex)
 					{
@@ -907,7 +907,7 @@ void FDecompressionTools::GetAnimationAttributes(const UAnimSequence* AnimSequen
 	}
 }
 
-void FDecompressionTools::GetAnimationAttributes(const UAnimSequence* AnimSequence, const FAnimExtractContext& ExtractionContext, const FEvaluationTaskContext& VMContext, FAttributeCollection& OutCollection, bool bForceUseRawData)
+void FDecompressionTools::GetAnimationAttributes(const UAnimSequence* AnimSequence, const FAnimExtractContext& ExtractionContext, const FEvaluationTaskContext& VMContext, FValueBundle& OutCollection, bool bForceUseRawData)
 {
 	QUICK_SCOPE_CYCLE_COUNTER(STAT_EvaluateAttributes);
 
@@ -962,7 +962,7 @@ void FDecompressionTools::GetAnimationAttributes(const UAnimSequence* AnimSequen
 						{
 							// This attribute is present within the current named set
 
-							FBoundValueMap* Map = OutCollection.FindSetMap(FAttributeMappingKey::MakeFromTo(AttributeType));
+							FBoundValueMap* Map = OutCollection.GetBoundValueMaps().Find(FAttributeMappingKey::MakeFromTo(AttributeType));
 							checkf(Map != nullptr, TEXT("This attribute exists within our named set but isn't present in the output collection"));
 
 							Map->SetValueWithSetter(SetIndex, [&ExtractionContext, &BakedAttribute](UScriptStruct* ValueType, uint8* OutValuePtr)
@@ -978,7 +978,7 @@ void FDecompressionTools::GetAnimationAttributes(const UAnimSequence* AnimSequen
 					else
 					{
 						// This attribute type is unknown within our binding, it is dynamic
-						FUnboundValueMap* AttributeMap = OutCollection.FindOrAddMap(AttributeType);
+						FUnboundValueMap* AttributeMap = OutCollection.GetUnboundValueMaps().FindOrAdd(AttributeType);
 
 						AttributeMap->AddWithSetter(AttributeName, [&ExtractionContext, &BakedAttribute](UScriptStruct* ValueType, uint8* OutValuePtr)
 							{
diff --git a/Engine/Plugins/Experimental/UAF/UAFAnimGraph/Source/UAFAnimGraph/Private/EvaluationVM/Tasks/NormalizeRotations.cpp b/Engine/Plugins/Experimental/UAF/UAFAnimGraph/Source/UAFAnimGraph/Private/EvaluationVM/Tasks/NormalizeRotations.cpp
--- a/Engine/Plugins/Experimental/UAF/UAFAnimGraph/Source/UAFAnimGraph/Private/EvaluationVM/Tasks/NormalizeRotations.cpp
+++ b/Engine/Plugins/Experimental/UAF/UAFAnimGraph/Source/UAFAnimGraph/Private/EvaluationVM/Tasks/NormalizeRotations.cpp
@@ -5,7 +5,7 @@
 #include "EvaluationVM/EvaluationVM.h"
 #include "EvaluationVM/KeyframeState.h"
 #include "TransformArrayOperations.h"
-#include "UAF/ValueRuntime/AttributeCollection.h"
+#include "UAF/ValueRuntime/ValueBundle.h"
 #include "UAF/ValueRuntime/Transformers/Sanitize.h"
 
 #include UE_INLINE_GENERATED_CPP_BY_NAME(NormalizeRotations)
@@ -16,7 +16,7 @@ void FAnimNextNormalizeKeyframeRotationsTask::Execute(UE::UAF::FEvaluationVM& VM
 
 	if (VM.GetActiveNamedSet())
 	{
-		TUniquePtr<FAttributeCollection> Collection;
+		TUniquePtr<FValueBundle> Collection;
 		if (!VM.PopValue(ATTRIBUTE_STACK_NAME, Collection))
 		{
 			// We have no inputs, nothing to do
diff --git a/Engine/Plugins/Experimental/UAF/UAFAnimGraph/Source/UAFAnimGraph/Private/EvaluationVM/Tasks/PushAnimSequenceKeyframe.cpp b/Engine/Plugins/Experimental/UAF/UAFAnimGraph/Source/UAFAnimGraph/Private/EvaluationVM/Tasks/PushAnimSequenceKeyframe.cpp
--- a/Engine/Plugins/Experimental/UAF/UAFAnimGraph/Source/UAFAnimGraph/Private/EvaluationVM/Tasks/PushAnimSequenceKeyframe.cpp
+++ b/Engine/Plugins/Experimental/UAF/UAFAnimGraph/Source/UAFAnimGraph/Private/EvaluationVM/Tasks/PushAnimSequenceKeyframe.cpp
@@ -57,7 +57,7 @@ void FAnimNextAnimSequenceKeyframeTask::Execute(UE::UAF::FEvaluationVM& VM) cons
 				ValueSpace = FValueSpace(EValueSpaceType::Local, bIsAdditive);
 			}
 
-			FAttributeCollectionStack Collection(VMContext.GetNamedSet());
+			FValueBundleStack Collection(VMContext.GetNamedSet());
 			Collection.InitWithValueSpace(ValueSpace);
 
 			FDecompressionTools::GetAnimationPose(AnimSequencePtr, ExtractionContext, VMContext, Collection, bUseRawData);
@@ -76,7 +76,7 @@ void FAnimNextAnimSequenceKeyframeTask::Execute(UE::UAF::FEvaluationVM& VM) cons
 						// This attribute is present within the current named set
 
 						FAttributeMappingKey MappingKey = FAttributeMappingKey::MakeFromTo<FTransformAnimationAttribute>();
-						TBoundValueMap<FTransformAnimationAttribute>* Map = Collection.FindSetMap<FTransformAnimationAttribute>(MappingKey);
+						TBoundValueMap<FTransformAnimationAttribute>* Map = Collection.GetBoundValueMaps().Find<FTransformAnimationAttribute>(MappingKey);
 						checkf(Map != nullptr, TEXT("This attribute exists within our named set but isn't present in the output collection"));
 
 						const FTransform RootMotionTransform = AnimSequence->ExtractRootMotion(ExtractionContext);
@@ -85,7 +85,7 @@ void FAnimNextAnimSequenceKeyframeTask::Execute(UE::UAF::FEvaluationVM& VM) cons
 				}
 			}
 
-			VM.PushValue(ATTRIBUTE_STACK_NAME, MakeUnique<FAttributeCollection>(MoveTemp(Collection)));
+			VM.PushValue(ATTRIBUTE_STACK_NAME, MakeUnique<FValueBundle>(MoveTemp(Collection)));
 		}
 		else
 		{
@@ -128,10 +128,10 @@ void FAnimNextAnimSequenceKeyframeTask::Execute(UE::UAF::FEvaluationVM& VM) cons
 	{
 		if (VMContext.IsValid())
 		{
-			FAttributeCollectionStack Collection(VMContext.GetNamedSet());
+			FValueBundleStack Collection(VMContext.GetNamedSet());
 			Collection.InitWithValueSpace(FValueSpace(EValueSpaceType::Local));
 
-			VM.PushValue(ATTRIBUTE_STACK_NAME, MakeUnique<FAttributeCollection>(MoveTemp(Collection)));
+			VM.PushValue(ATTRIBUTE_STACK_NAME, MakeUnique<FValueBundle>(MoveTemp(Collection)));
 		}
 		else
 		{
diff --git a/Engine/Plugins/Experimental/UAF/UAFAnimGraph/Source/UAFAnimGraph/Private/TraitInterfaces/IEvaluate.cpp b/Engine/Plugins/Experimental/UAF/UAFAnimGraph/Source/UAFAnimGraph/Private/TraitInterfaces/IEvaluate.cpp
--- a/Engine/Plugins/Experimental/UAF/UAFAnimGraph/Source/UAFAnimGraph/Private/TraitInterfaces/IEvaluate.cpp
+++ b/Engine/Plugins/Experimental/UAF/UAFAnimGraph/Source/UAFAnimGraph/Private/TraitInterfaces/IEvaluate.cpp
@@ -13,7 +13,7 @@
 #include "Graph/TraceAnimNextGraphInstances.h"
 
 #include "UAF/Attributes/EngineAttributes.h"
-#include "UAF/ValueRuntime/AttributeCollection.h"
+#include "UAF/ValueRuntime/ValueBundle.h"
 
 namespace UE::UAF
 {
@@ -255,7 +255,7 @@ namespace UE::UAF
 			{
 				EvaluationProgram.Execute(EvaluationVM);
 
-				TUniquePtr<FAttributeCollection> Collection;
+				TUniquePtr<FValueBundle> Collection;
 				if (Private::GUseAbstractSkeletonRuntime && EvaluationVM.PopValue(ATTRIBUTE_STACK_NAME, Collection))
 				{
 					const FAttributeMappingKey MappingKey = FAttributeMappingKey::MakeFromTo<FBoneTransformAnimationAttribute>();
@@ -266,8 +266,8 @@ namespace UE::UAF
 						// Attach to root for now for backwards compatibility
 						const FCompactPoseBoneIndex RootCompactPoseBoneIndex(0);
 
-						// Set maps
-						for (auto It = Collection->CreateConstSetMapIterator(); It; ++It)
+						// Bound value maps
+						for (auto It = Collection->GetBoundValueMaps().CreateIterator(); It; ++It)
 						{
 							const FBoundValueMap* Map = It.GetMap();
 							const FAttributeTypedSetPtr& TypedSet = Map->GetTypedSet();
@@ -329,8 +329,8 @@ namespace UE::UAF
 							}
 						}
 
-						// Name/value maps
-						for (auto It = Collection->CreateConstMapIterator(); It; ++It)
+						// Unbound value maps
+						for (auto It = Collection->GetUnboundValueMaps().CreateConstIterator(); It; ++It)
 						{
 							const FUnboundValueMap* Map = It.GetMap();
 							const UScriptStruct* ValueType = Map->GetValueType();
diff --git a/Engine/Plugins/Experimental/UAF/UAFAnimGraph/Source/UAFAnimGraph/Private/Traits/ModifyCurveTrait.cpp b/Engine/Plugins/Experimental/UAF/UAFAnimGraph/Source/UAFAnimGraph/Private/Traits/ModifyCurveTrait.cpp
--- a/Engine/Plugins/Experimental/UAF/UAFAnimGraph/Source/UAFAnimGraph/Private/Traits/ModifyCurveTrait.cpp
+++ b/Engine/Plugins/Experimental/UAF/UAFAnimGraph/Source/UAFAnimGraph/Private/Traits/ModifyCurveTrait.cpp
@@ -53,14 +53,14 @@ void FModifyCurveTask::Execute(UE::UAF::FEvaluationVM& VM) const
 
 	if (VM.GetActiveNamedSet())
 	{
-		if (const TUniquePtr<FAttributeCollection>* Collection = VM.PeekValue<TUniquePtr<FAttributeCollection>>(ATTRIBUTE_STACK_NAME, 0))
+		if (const TUniquePtr<FValueBundle>* Collection = VM.PeekValue<TUniquePtr<FValueBundle>>(ATTRIBUTE_STACK_NAME, 0))
 		{
 			FAttributeNamedSetPtr NamedSet = (*Collection)->GetNamedSet();
 			if (FAttributeTypedSetPtr FloatTypedSet = NamedSet->FindTypedSet<FFloatAnimationAttribute>())
 			{
 				FAttributeMappingKey MappingKey = FAttributeMappingKey::MakeFromTo<FFloatAnimationAttribute>();
 				
-				if (TUnboundValueMap<FFloatAnimationAttribute>* FloatMap = (*Collection)->FindMap<FFloatAnimationAttribute>())
+				if (TUnboundValueMap<FFloatAnimationAttribute>* FloatMap = (*Collection)->GetUnboundValueMaps().Find<FFloatAnimationAttribute>())
 				{
 					for (const FModifyCurveParameters& CurveParameters : SharedData->ModifyCurveParameters)
 					{
@@ -73,7 +73,7 @@ void FModifyCurveTask::Execute(UE::UAF::FEvaluationVM& VM) const
 					}
 				}
 
-				if (TBoundValueMap<FFloatAnimationAttribute>* FloatMap = (*Collection)->FindSetMap<FFloatAnimationAttribute>(MappingKey))
+				if (TBoundValueMap<FFloatAnimationAttribute>* FloatMap = (*Collection)->GetBoundValueMaps().Find<FFloatAnimationAttribute>(MappingKey))
 				{
 					for (const FModifyCurveParameters& CurveParameters : SharedData->ModifyCurveParameters)
 					{
diff --git a/Engine/Plugins/Experimental/UAF/UAFAnimGraph/Source/UAFAnimGraph/Public/EvaluationVM/EvaluationVM.h b/Engine/Plugins/Experimental/UAF/UAFAnimGraph/Source/UAFAnimGraph/Public/EvaluationVM/EvaluationVM.h
--- a/Engine/Plugins/Experimental/UAF/UAFAnimGraph/Source/UAFAnimGraph/Public/EvaluationVM/EvaluationVM.h
+++ b/Engine/Plugins/Experimental/UAF/UAFAnimGraph/Source/UAFAnimGraph/Public/EvaluationVM/EvaluationVM.h
@@ -13,7 +13,7 @@
 #include "ReferencePose.h"
 #include "Misc/Crc.h"
 #include "UAF/AbstractSkeleton/AbstractSkeletonSetBinding.h"
-#include "UAF/ValueRuntime/AttributeCollection.h"
+#include "UAF/ValueRuntime/ValueBundle.h"
 #include "UAF/ValueRuntime/Transformers/Transformer.h"
 #include "UAF/ValueRuntime/ValueRuntimeRegistry.h"
 
@@ -341,9 +341,9 @@ namespace UE::UAF
 	extern UAFANIMGRAPH_API const FEvaluationVMStackName KEYFRAME_STACK_NAME;
 	ANIM_NEXT_ENABLE_EVALUATION_STACK_USAGE(TUniquePtr<FKeyframeState>)
 
-	// A stack of FAttributeCollection instances used when sampling sequences and blending their results
+	// A stack of FValueBundle instances used when sampling sequences and blending their results
 	extern UAFANIMGRAPH_API const FEvaluationVMStackName ATTRIBUTE_STACK_NAME;
-	ANIM_NEXT_ENABLE_EVALUATION_STACK_USAGE(TUniquePtr<FAttributeCollection>)
+	ANIM_NEXT_ENABLE_EVALUATION_STACK_USAGE(TUniquePtr<FValueBundle>)
 }
 
 #undef UE_API
diff --git a/Engine/Plugins/Experimental/UAF/UAFAnimGraph/Source/UAFAnimGraph/Public/EvaluationVM/Tasks/ApplyAdditiveKeyframe.h b/Engine/Plugins/Experimental/UAF/UAFAnimGraph/Source/UAFAnimGraph/Public/EvaluationVM/Tasks/ApplyAdditiveKeyframe.h
--- a/Engine/Plugins/Experimental/UAF/UAFAnimGraph/Source/UAFAnimGraph/Public/EvaluationVM/Tasks/ApplyAdditiveKeyframe.h
+++ b/Engine/Plugins/Experimental/UAF/UAFAnimGraph/Source/UAFAnimGraph/Public/EvaluationVM/Tasks/ApplyAdditiveKeyframe.h
@@ -13,7 +13,7 @@ struct FInputScaleBiasClamp;
 
 namespace UE::UAF
 {
-	class FAttributeCollection;
+	class FValueBundle;
 	struct FKeyframeState;
 }
 
@@ -40,7 +40,7 @@ struct FAnimNextApplyAdditiveKeyframeTask : public FAnimNextEvaluationTask
 	// @TODO: Move this into shared method
 	// Compute the current alpha value for curves or return the interpolation alpha otherwise
 	UE_API float GetInterpolationAlpha(const UE::UAF::FKeyframeState* KeyframeA, const UE::UAF::FKeyframeState* KeyframeB) const;
-	UE_API float GetInterpolationAlpha(const UE::UAF::FAttributeCollection& CollectionA, const UE::UAF::FAttributeCollection& CollectionB) const;
+	UE_API float GetInterpolationAlpha(const UE::UAF::FValueBundle& CollectionA, const UE::UAF::FValueBundle& CollectionB) const;
 
 	// How much weight between the additive identity and the additive pose to apply: lerp(identity, additive, weight)
 	UPROPERTY()
diff --git a/Engine/Plugins/Experimental/UAF/UAFAnimGraph/Source/UAFAnimGraph/Public/EvaluationVM/Tasks/BlendKeyframes.h b/Engine/Plugins/Experimental/UAF/UAFAnimGraph/Source/UAFAnimGraph/Public/EvaluationVM/Tasks/BlendKeyframes.h
--- a/Engine/Plugins/Experimental/UAF/UAFAnimGraph/Source/UAFAnimGraph/Public/EvaluationVM/Tasks/BlendKeyframes.h
+++ b/Engine/Plugins/Experimental/UAF/UAFAnimGraph/Source/UAFAnimGraph/Public/EvaluationVM/Tasks/BlendKeyframes.h
@@ -13,7 +13,7 @@ struct FInputScaleBiasClamp;
 
 namespace UE::UAF
 {
-	class FAttributeCollection;
+	class FValueBundle;
 	struct FKeyframeState;
 }
 
@@ -53,7 +53,7 @@ struct FAnimNextBlendTwoKeyframesTask : public FAnimNextEvaluationTask
 
 	// Compute the current alpha value for curves or return the interpolation alpha otherwise
 	UE_API float GetInterpolationAlpha(const UE::UAF::FKeyframeState* KeyframeA, const UE::UAF::FKeyframeState* KeyframeB) const;
-	UE_API float GetInterpolationAlpha(const UE::UAF::FAttributeCollection& CollectionA, const UE::UAF::FAttributeCollection& CollectionB) const;
+	UE_API float GetInterpolationAlpha(const UE::UAF::FValueBundle& CollectionA, const UE::UAF::FValueBundle& CollectionB) const;
 
 	// The interpolation alpha between the two input keyframes
 	UPROPERTY(VisibleAnywhere, Category = Properties)
diff --git a/Engine/Plugins/Experimental/UAF/UAFAnimGraph/Source/UAFAnimGraph/Public/EvaluationVM/Tasks/DecompressionTools.h b/Engine/Plugins/Experimental/UAF/UAFAnimGraph/Source/UAFAnimGraph/Public/EvaluationVM/Tasks/DecompressionTools.h
--- a/Engine/Plugins/Experimental/UAF/UAFAnimGraph/Source/UAFAnimGraph/Public/EvaluationVM/Tasks/DecompressionTools.h
+++ b/Engine/Plugins/Experimental/UAF/UAFAnimGraph/Source/UAFAnimGraph/Public/EvaluationVM/Tasks/DecompressionTools.h
@@ -19,7 +19,7 @@ class UAnimSequence;
 
 namespace UE::UAF
 {
-	class FAttributeCollection;
+	class FValueBundle;
 	class FEvaluationTaskContext;
 
 class FDecompressionTools
@@ -31,7 +31,7 @@ public:
 	// Extracted from UAnimSequence
 	// Extracts Animation Data from the provided AnimSequence, using AnimExtractContext extraction parameters
 	static UE_API void GetAnimationPose(const UAnimSequence* AnimSequence, const FAnimExtractContext& ExtractionContext, FLODPose& OutAnimationPoseData, bool bForceUseRawData = false);
-	static UE_API void GetAnimationPose(const UAnimSequence* AnimSequence, const FAnimExtractContext& ExtractionContext, const FEvaluationTaskContext& VMContext, FAttributeCollection& OutCollection, bool bForceUseRawData = false);
+	static UE_API void GetAnimationPose(const UAnimSequence* AnimSequence, const FAnimExtractContext& ExtractionContext, const FEvaluationTaskContext& VMContext, FValueBundle& OutCollection, bool bForceUseRawData = false);
 
 	/**
 	* Get Bone Transform of the Time given, relative to Parent for all RequiredBones
@@ -43,10 +43,10 @@ public:
 	static UE_API void GetBonePose_AdditiveMeshRotationOnly(const UAnimSequence* AnimSequence, const FAnimExtractContext& ExtractionContext, FLODPose& OutAnimationPoseData);
 
 	static UE_API void GetAnimationCurves(const UAnimSequence* AnimSequence, const FAnimExtractContext& ExtractionContext, FBlendedCurve& OutCurves, bool bForceUseRawData = false);
-	static UE_API void GetAnimationCurves(const UAnimSequence* AnimSequence, const FAnimExtractContext& ExtractionContext, const FEvaluationTaskContext& VMContext, FAttributeCollection& OutCollection, bool bForceUseRawData = false);
+	static UE_API void GetAnimationCurves(const UAnimSequence* AnimSequence, const FAnimExtractContext& ExtractionContext, const FEvaluationTaskContext& VMContext, FValueBundle& OutCollection, bool bForceUseRawData = false);
 
 	static UE_API void GetAnimationAttributes(const UAnimSequence* AnimSequence, const FAnimExtractContext& ExtractionContext, const FReferencePose& RefPose, UE::Anim::FStackAttributeContainer& OutAttributes, bool bForceUseRawData = false);
-	static UE_API void GetAnimationAttributes(const UAnimSequence* AnimSequence, const FAnimExtractContext& ExtractionContext, const FEvaluationTaskContext& VMContext, FAttributeCollection& OutCollection, bool bForceUseRawData = false);
+	static UE_API void GetAnimationAttributes(const UAnimSequence* AnimSequence, const FAnimExtractContext& ExtractionContext, const FEvaluationTaskContext& VMContext, FValueBundle& OutCollection, bool bForceUseRawData = false);
 
 	// Decompress and retarget animation data using provided RetargetTransforms
 	static UE_API void DecompressPose(FLODPose& OutAnimationPoseData,
