Rename UAnimNextAnimationGraph_EditorData to UUAFAnimGraph_EditorData
Rename UAnimNextAnimationGraphFactory to UUAFAnimGraphFactory

#Synced-CL 49648897

diff --git a/Engine/Plugins/Experimental/UAF/UAF/Config/DefaultUAF.ini b/Engine/Plugins/Experimental/UAF/UAF/Config/DefaultUAF.ini
--- a/Engine/Plugins/Experimental/UAF/UAF/Config/DefaultUAF.ini
+++ b/Engine/Plugins/Experimental/UAF/UAF/Config/DefaultUAF.ini
@@ -79,10 +79,10 @@
 +StructRedirects=(OldName="/Script/AnimNext.AnimNextTraitSharedData", NewName="/Script/UAFAnimGraph.AnimNextTraitSharedData")
 +StructRedirects=(OldName="/Script/AnimNext.AnimNextGraphInstance", NewName="/Script/UAFAnimGraph.AnimNextGraphInstance")
 
-+ClassRedirects=(OldName="/Script/AnimNextEditor.AnimNextAnimationGraphFactory", NewName="/Script/UAFAnimGraphEditor.AnimNextAnimationGraphFactory")
++ClassRedirects=(OldName="/Script/AnimNextEditor.AnimNextAnimationGraphFactory", NewName="/Script/UAFAnimGraphEditor.UAFAnimGraphFactory")
 +ClassRedirects=(OldName="/Script/AnimNextEditor.AssetDefinition_AnimNextAnimationGraph", NewName="/Script/UAFAnimGraphEditor.AssetDefinition_AnimNextAnimationGraph")
 
-+ClassRedirects=(OldName="/Script/AnimNextUncookedOnly.AnimNextAnimationGraph_EditorData", NewName="/Script/UAFAnimGraphUncookedOnly.AnimNextAnimationGraph_EditorData")
++ClassRedirects=(OldName="/Script/AnimNextUncookedOnly.AnimNextAnimationGraph_EditorData", NewName="/Script/UAFAnimGraphUncookedOnly.UAFAnimGraph_EditorData")
 +ClassRedirects=(OldName="/Script/AnimNextUncookedOnly.AnimNextAnimationGraphSchema", NewName="/Script/UAFAnimGraphUncookedOnly.AnimNextAnimationGraphSchema")
 +StructRedirects=(OldName="/Script/AnimNextUncookedOnly.AnimNextAnimationGraphOutlinerData", NewName="/Script/UAFAnimGraphUncookedOnly.AnimNextAnimationGraphOutlinerData")
 +ClassRedirects=(OldName="/Script/AnimNextUncookedOnly.AnimNextAnimationGraphEntry", NewName="/Script/UAFAnimGraphUncookedOnly.AnimNextAnimationGraphEntry")
diff --git a/Engine/Plugins/Experimental/UAF/UAFAnimGraph/Config/DefaultUAFAnimGraph.ini b/Engine/Plugins/Experimental/UAF/UAFAnimGraph/Config/DefaultUAFAnimGraph.ini
--- a/Engine/Plugins/Experimental/UAF/UAFAnimGraph/Config/DefaultUAFAnimGraph.ini
+++ b/Engine/Plugins/Experimental/UAF/UAFAnimGraph/Config/DefaultUAFAnimGraph.ini
@@ -168,7 +168,7 @@
 +StructRedirects=(OldName="/Script/AnimNextAnimGraphEditor.AnimNextSyncMarkerTraceInfo", NewName="/Script/UAFAnimGraphEditor.AnimNextSyncMarkerTraceInfo")
 +StructRedirects=(OldName="/Script/AnimNextAnimGraphEditor.AnimNextSequenceTraceInfo", NewName="/Script/UAFAnimGraphEditor.AnimNextSequenceTraceInfo")
 
-+ClassRedirects=(OldName="/Script/AnimNextAnimGraphEditor.AnimNextAnimationGraphFactory", NewName="/Script/UAFAnimGraphEditor.AnimNextAnimationGraphFactory")
++ClassRedirects=(OldName="/Script/AnimNextAnimGraphEditor.AnimNextAnimationGraphFactory", NewName="/Script/UAFAnimGraphEditor.UAFAnimGraphFactory")
 +ClassRedirects=(OldName="/Script/AnimNextAnimGraphEditor.AssetDefinition_AnimNextAnimationGraphEntry", NewName="/Script/UAFAnimGraphEditor.AssetDefinition_AnimNextAnimationGraphEntry")
 +ClassRedirects=(OldName="/Script/AnimNextAnimGraphEditor.AssetDefinition_AnimNextAnimationGraph", NewName="/Script/UAFAnimGraphEditor.AssetDefinition_AnimNextAnimationGraph")
 +ClassRedirects=(OldName="/Script/AnimNextAnimGraphEditor.EvaluationProgramDetailsObject", NewName="/Script/UAFAnimGraphEditor.EvaluationProgramDetailsObject")
@@ -199,10 +199,12 @@
 +ClassRedirects=(OldName="/Script/AnimNextAnimGraphUncookedOnly.AnimGraphNode_AnimNextGraph", NewName="/Script/UAFAnimGraphUncookedOnly.AnimGraphNode_AnimNextGraph")
 +ClassRedirects=(OldName="/Script/AnimNextAnimGraphUncookedOnly.AnimNextAnimationGraphSchema", NewName="/Script/UAFAnimGraphUncookedOnly.AnimNextAnimationGraphSchema")
 +ClassRedirects=(OldName="/Script/AnimNextAnimGraphUncookedOnly.AnimNextAnimationGraphLibrary", NewName="/Script/UAFAnimGraphUncookedOnly.AnimNextAnimationGraphLibrary")
-+ClassRedirects=(OldName="/Script/AnimNextAnimGraphUncookedOnly.AnimNextAnimationGraph_EditorData", NewName="/Script/UAFAnimGraphUncookedOnly.AnimNextAnimationGraph_EditorData")
++ClassRedirects=(OldName="/Script/AnimNextAnimGraphUncookedOnly.AnimNextAnimationGraph_EditorData", NewName="/Script/UAFAnimGraphUncookedOnly.UAFAnimGraph_EditorData")
 +ClassRedirects=(OldName="/Script/AnimNextAnimGraphUncookedOnly.K2Node_AnimNextInjection", NewName="/Script/UAFAnimGraphUncookedOnly.K2Node_AnimNextInjection")
 +ClassRedirects=(OldName="/Script/AnimNextAnimGraphUncookedOnly.K2Node_AnimNextPlayAnim", NewName="/Script/UAFAnimGraphUncookedOnly.K2Node_AnimNextPlayAnim")
 +ClassRedirects=(OldName="/Script/AnimNextAnimGraphUncookedOnly.K2Node_InjectionCallbackProxySetVariable", NewName="/Script/UAFAnimGraphUncookedOnly.K2Node_InjectionCallbackProxySetVariable")
 
 ; AnimNext -> UAF
-+ClassRedirects=(OldName="/Script/UAFAnimGraph.AnimNextAnimationGraph",NewName="/Script/UAFAnimGraph.UAFAnimGraph")
\ No newline at end of file
++ClassRedirects=(OldName="/Script/UAFAnimGraph.AnimNextAnimationGraph",NewName="/Script/UAFAnimGraph.UAFAnimGraph")
++ClassRedirects=(OldName="/Script/UAFAnimGraphUncookedOnly.AnimNextAnimationGraph_EditorData",NewName="/Script/UAFAnimGraphUncookedOnly.UAFAnimGraph_EditorData")
++ClassRedirects=(OldName="/Script/UAFAnimGraphUncookedOnly.AnimNextAnimationGraphFactory",NewName="/Script/UAFAnimGraphUncookedOnly.UAFAnimGraphFactory")
\ No newline at end of file
diff --git a/Engine/Plugins/Experimental/UAF/UAFAnimGraph/Source/UAFAnimGraph/Public/Graph/AnimNextAnimationGraph.h b/Engine/Plugins/Experimental/UAF/UAFAnimGraph/Source/UAFAnimGraph/Public/Graph/AnimNextAnimationGraph.h
--- a/Engine/Plugins/Experimental/UAF/UAFAnimGraph/Source/UAFAnimGraph/Public/Graph/AnimNextAnimationGraph.h
+++ b/Engine/Plugins/Experimental/UAF/UAFAnimGraph/Source/UAFAnimGraph/Public/Graph/AnimNextAnimationGraph.h
@@ -123,8 +123,8 @@ private:
 	UE_API void ThawGraphInstances();
 #endif
 	
-	friend class UAnimNextAnimationGraphFactory;
-	friend class UAnimNextAnimationGraph_EditorData;
+	friend class UUAFAnimGraphFactory;
+	friend class UUAFAnimGraph_EditorData;
 	friend class UAnimNextVariableEntry;
 	friend struct UE::UAF::UncookedOnly::FUtils;
 	friend class UE::UAF::Editor::FModuleEditor;
diff --git a/Engine/Plugins/Experimental/UAF/UAFAnimGraph/Source/UAFAnimGraphEditor/Internal/Graph/AnimNextAnimationGraphFactory.h b/Engine/Plugins/Experimental/UAF/UAFAnimGraph/Source/UAFAnimGraphEditor/Internal/Graph/AnimNextAnimationGraphFactory.h
--- a/Engine/Plugins/Experimental/UAF/UAFAnimGraph/Source/UAFAnimGraphEditor/Internal/Graph/AnimNextAnimationGraphFactory.h
+++ b/Engine/Plugins/Experimental/UAF/UAFAnimGraph/Source/UAFAnimGraphEditor/Internal/Graph/AnimNextAnimationGraphFactory.h
@@ -7,11 +7,11 @@
 #include "AnimNextAnimationGraphFactory.generated.h"
 
 UCLASS(MinimalAPI)
-class UAnimNextAnimationGraphFactory : public UFactory
+class UUAFAnimGraphFactory : public UFactory
 {
 	GENERATED_BODY()
 
-	UAnimNextAnimationGraphFactory();
+	UUAFAnimGraphFactory();
 
 	virtual bool ConfigureProperties() override;
 	virtual UObject* FactoryCreateNew(UClass* Class, UObject* InParent, FName Name, EObjectFlags Flags, UObject* Context, FFeedbackContext* Warn, FName CallingContext) override;
diff --git a/Engine/Plugins/Experimental/UAF/UAFAnimGraph/Source/UAFAnimGraphEditor/Private/AnimNextAnimGraphEditorModule.cpp b/Engine/Plugins/Experimental/UAF/UAFAnimGraph/Source/UAFAnimGraphEditor/Private/AnimNextAnimGraphEditorModule.cpp
--- a/Engine/Plugins/Experimental/UAF/UAFAnimGraph/Source/UAFAnimGraphEditor/Private/AnimNextAnimGraphEditorModule.cpp
+++ b/Engine/Plugins/Experimental/UAF/UAFAnimGraph/Source/UAFAnimGraphEditor/Private/AnimNextAnimGraphEditorModule.cpp
@@ -128,7 +128,7 @@ class FAnimNextAnimGraphEditorModule : public IModuleInterface
 				URigVMEdGraph* EdGraph = nullptr;
 				
 				UUAFAnimGraph* AnimationGraph = CastChecked<UUAFAnimGraph>(Document.GetObject());
-				UAnimNextAnimationGraph_EditorData* EditorData = UncookedOnly::FUtils::GetEditorData<UAnimNextAnimationGraph_EditorData>(AnimationGraph);
+				UUAFAnimGraph_EditorData* EditorData = UncookedOnly::FUtils::GetEditorData<UUAFAnimGraph_EditorData>(AnimationGraph);
 
 				UAnimNextAnimationGraphEntry* AnimationGraphEntry = Cast<UAnimNextAnimationGraphEntry>(EditorData->FindEntry(FRigUnit_AnimNextGraphRoot::DefaultEntryPoint));
 				// Redirect to the inner graph
diff --git a/Engine/Plugins/Experimental/UAF/UAFAnimGraph/Source/UAFAnimGraphEditor/Private/Graph/AnimNextAnimationGraphFactory.cpp b/Engine/Plugins/Experimental/UAF/UAFAnimGraph/Source/UAFAnimGraphEditor/Private/Graph/AnimNextAnimationGraphFactory.cpp
--- a/Engine/Plugins/Experimental/UAF/UAFAnimGraph/Source/UAFAnimGraphEditor/Private/Graph/AnimNextAnimationGraphFactory.cpp
+++ b/Engine/Plugins/Experimental/UAF/UAFAnimGraph/Source/UAFAnimGraphEditor/Private/Graph/AnimNextAnimationGraphFactory.cpp
@@ -11,19 +11,19 @@
 
 #include UE_INLINE_GENERATED_CPP_BY_NAME(AnimNextAnimationGraphFactory)
 
-UAnimNextAnimationGraphFactory::UAnimNextAnimationGraphFactory()
+UUAFAnimGraphFactory::UUAFAnimGraphFactory()
 {
 	bCreateNew = true;
 	bEditAfterNew = true;
 	SupportedClass = UUAFAnimGraph::StaticClass();
 }
 
-bool UAnimNextAnimationGraphFactory::ConfigureProperties()
+bool UUAFAnimGraphFactory::ConfigureProperties()
 {
 	return true;
 }
 
-UObject* UAnimNextAnimationGraphFactory::FactoryCreateNew(UClass* Class, UObject* InParent, FName Name, EObjectFlags Flags, UObject* Context, FFeedbackContext* Warn, FName CallingContext)
+UObject* UUAFAnimGraphFactory::FactoryCreateNew(UClass* Class, UObject* InParent, FName Name, EObjectFlags Flags, UObject* Context, FFeedbackContext* Warn, FName CallingContext)
 {
 	EObjectFlags FlagsToUse = Flags | RF_Public | RF_Standalone | RF_Transactional | RF_LoadCompleted;
 	if(InParent == GetTransientPackage())
@@ -34,7 +34,7 @@ UObject* UAnimNextAnimationGraphFactory::FactoryCreateNew(UClass* Class, UObject
 	UUAFAnimGraph* NewAnimationGraph = NewObject<UUAFAnimGraph>(InParent, Class, Name, FlagsToUse);
 
 	// Create internal editor data
-	UAnimNextAnimationGraph_EditorData* EditorData = NewObject<UAnimNextAnimationGraph_EditorData>(NewAnimationGraph, TEXT("EditorData"), RF_Transactional);
+	UUAFAnimGraph_EditorData* EditorData = NewObject<UUAFAnimGraph_EditorData>(NewAnimationGraph, TEXT("EditorData"), RF_Transactional);
 	NewAnimationGraph->EditorData = EditorData;
 	EditorData->bUsesExternalPackages = false;
 	EditorData->Initialize(/*bRecompileVM*/false);
diff --git a/Engine/Plugins/Experimental/UAF/UAFAnimGraph/Source/UAFAnimGraphUncookedOnly/Internal/AnimNextTraitStackUnitNode.h b/Engine/Plugins/Experimental/UAF/UAFAnimGraph/Source/UAFAnimGraphUncookedOnly/Internal/AnimNextTraitStackUnitNode.h
--- a/Engine/Plugins/Experimental/UAF/UAFAnimGraph/Source/UAFAnimGraphUncookedOnly/Internal/AnimNextTraitStackUnitNode.h
+++ b/Engine/Plugins/Experimental/UAF/UAFAnimGraph/Source/UAFAnimGraphUncookedOnly/Internal/AnimNextTraitStackUnitNode.h
@@ -10,7 +10,7 @@
 #define UE_API UAFANIMGRAPHUNCOOKEDONLY_API
 
 class UUAFGraphNodeTemplate;
-class UAnimNextAnimationGraph_EditorData;
+class UUAFAnimGraph_EditorData;
 class UAnimNextController;
 
 namespace UE::UAF::Editor
@@ -38,7 +38,7 @@ public:
 
 private:
 	friend UUAFGraphNodeTemplate;
-	friend UAnimNextAnimationGraph_EditorData;
+	friend UUAFAnimGraph_EditorData;
 	friend UE::UAF::Editor::FAnimationGraphMenuExtensions;
 
 	void HandlePinDefaultValueChanged(UAnimNextController* InController, URigVMPin* InPin);
diff --git a/Engine/Plugins/Experimental/UAF/UAFAnimGraph/Source/UAFAnimGraphUncookedOnly/Internal/Entries/AnimNextAnimationGraphEntry.h b/Engine/Plugins/Experimental/UAF/UAFAnimGraph/Source/UAFAnimGraphUncookedOnly/Internal/Entries/AnimNextAnimationGraphEntry.h
--- a/Engine/Plugins/Experimental/UAF/UAFAnimGraph/Source/UAFAnimGraphUncookedOnly/Internal/Entries/AnimNextAnimationGraphEntry.h
+++ b/Engine/Plugins/Experimental/UAF/UAFAnimGraph/Source/UAFAnimGraphUncookedOnly/Internal/Entries/AnimNextAnimationGraphEntry.h
@@ -32,7 +32,7 @@ class UAnimNextAnimationGraphEntry : public UAnimNextRigVMAssetEntry, public IAn
 {
 	GENERATED_BODY()
 
-	friend class UAnimNextAnimationGraph_EditorData;
+	friend class UUAFAnimGraph_EditorData;
 	friend struct UE::UAF::Editor::FUtils;
 	friend class UE::UAF::Editor::FAnimNextAnimGraphEditorModule;
 	friend class UE::UAF::Tests::FEditor_Graphs;
diff --git a/Engine/Plugins/Experimental/UAF/UAFAnimGraph/Source/UAFAnimGraphUncookedOnly/Internal/Graph/AnimNextAnimationGraph_EditorData.h b/Engine/Plugins/Experimental/UAF/UAFAnimGraph/Source/UAFAnimGraphUncookedOnly/Internal/Graph/AnimNextAnimationGraph_EditorData.h
--- a/Engine/Plugins/Experimental/UAF/UAFAnimGraph/Source/UAFAnimGraphUncookedOnly/Internal/Graph/AnimNextAnimationGraph_EditorData.h
+++ b/Engine/Plugins/Experimental/UAF/UAFAnimGraph/Source/UAFAnimGraphUncookedOnly/Internal/Graph/AnimNextAnimationGraph_EditorData.h
@@ -40,13 +40,13 @@ public:
 	static UAFANIMGRAPHUNCOOKEDONLY_API UAnimNextAnimationGraphEntry* AddAnimationGraph(UUAFAnimGraph* InAsset, FName InName, bool bSetupUndoRedo = true, bool bPrintPythonCommand = true);
 };
 
-/** Editor data for AnimNext animation graphs */
+/** Editor data for UAF animation graphs */
 UCLASS(MinimalAPI)
-class UAnimNextAnimationGraph_EditorData : public UAnimNextSharedVariables_EditorData
+class UUAFAnimGraph_EditorData : public UAnimNextSharedVariables_EditorData
 {
 	GENERATED_BODY()
 
-	friend class UAnimNextAnimationGraphFactory;
+	friend class UUAFAnimGraphFactory;
 	friend class UAnimNextEdGraph;
 	friend struct UE::UAF::UncookedOnly::FUtils;
 	friend struct UE::UAF::Editor::FUtils;
diff --git a/Engine/Plugins/Experimental/UAF/UAFAnimGraph/Source/UAFAnimGraphUncookedOnly/Private/AnimNextAnimGraphWorkspaceAssetUserData.cpp b/Engine/Plugins/Experimental/UAF/UAFAnimGraph/Source/UAFAnimGraphUncookedOnly/Private/AnimNextAnimGraphWorkspaceAssetUserData.cpp
--- a/Engine/Plugins/Experimental/UAF/UAFAnimGraph/Source/UAFAnimGraphUncookedOnly/Private/AnimNextAnimGraphWorkspaceAssetUserData.cpp
+++ b/Engine/Plugins/Experimental/UAF/UAFAnimGraph/Source/UAFAnimGraphUncookedOnly/Private/AnimNextAnimGraphWorkspaceAssetUserData.cpp
@@ -39,7 +39,7 @@ void UAnimNextAnimGraphWorkspaceAssetUserData::GetWorkspaceAssetExports(FAssetRe
 	Super::GetAssetRegistryTags(Context);
 
 	UUAFAnimGraph* Asset = CastChecked<UUAFAnimGraph>(GetOuter());
-	const UAnimNextAnimationGraph_EditorData* EditorData = UE::UAF::UncookedOnly::FUtils::GetEditorData<UAnimNextAnimationGraph_EditorData>(Asset);
+	const UUAFAnimGraph_EditorData* EditorData = UE::UAF::UncookedOnly::FUtils::GetEditorData<UUAFAnimGraph_EditorData>(Asset);
 
 	UE::UAF::UncookedOnly::FUtils::GetAssetWorkspaceExports(EditorData, CachedExports, Context);
 }
diff --git a/Engine/Plugins/Experimental/UAF/UAFAnimGraph/Source/UAFAnimGraphUncookedOnly/Private/Graph/AnimNextAnimationGraph_EditorData.cpp b/Engine/Plugins/Experimental/UAF/UAFAnimGraph/Source/UAFAnimGraphUncookedOnly/Private/Graph/AnimNextAnimationGraph_EditorData.cpp
--- a/Engine/Plugins/Experimental/UAF/UAFAnimGraph/Source/UAFAnimGraphUncookedOnly/Private/Graph/AnimNextAnimationGraph_EditorData.cpp
+++ b/Engine/Plugins/Experimental/UAF/UAFAnimGraph/Source/UAFAnimGraphUncookedOnly/Private/Graph/AnimNextAnimationGraph_EditorData.cpp
@@ -628,7 +628,7 @@ namespace UE::UAF::UncookedOnly::Private
 	}
 }
 
-void UAnimNextAnimationGraph_EditorData::OnPreCompileAsset(FRigVMCompileSettings& InSettings)
+void UUAFAnimGraph_EditorData::OnPreCompileAsset(FRigVMCompileSettings& InSettings)
 {
 	using namespace UE::UAF::UncookedOnly;
 
@@ -646,7 +646,7 @@ void UAnimNextAnimationGraph_EditorData::OnPreCompileAsset(FRigVMCompileSettings
 	AnimationGraph->DefaultEntryPoint = NAME_None;
 }
 
-void UAnimNextAnimationGraph_EditorData::BuildFunctionHeadersContext(const FRigVMCompileSettings& InSettings, FAnimNextGetFunctionHeaderCompileContext& OutCompileContext) const
+void UUAFAnimGraph_EditorData::BuildFunctionHeadersContext(const FRigVMCompileSettings& InSettings, FAnimNextGetFunctionHeaderCompileContext& OutCompileContext) const
 {
 	using namespace UE::UAF::UncookedOnly;
 	using namespace UE::UAF::UncookedOnly::Private;
@@ -703,7 +703,7 @@ void UAnimNextAnimationGraph_EditorData::BuildFunctionHeadersContext(const FRigV
 	}
 }
 
-void UAnimNextAnimationGraph_EditorData::OnPreCompileProcessGraphs(const FRigVMCompileSettings& InSettings, FAnimNextProcessGraphCompileContext& OutCompileContext)
+void UUAFAnimGraph_EditorData::OnPreCompileProcessGraphs(const FRigVMCompileSettings& InSettings, FAnimNextProcessGraphCompileContext& OutCompileContext)
 {
 	using namespace UE::UAF::UncookedOnly;
 	using namespace UE::UAF::UncookedOnly::Private;
@@ -839,7 +839,7 @@ void UAnimNextAnimationGraph_EditorData::OnPreCompileProcessGraphs(const FRigVMC
 	InOutGraphs.Append(NonAnimGraphs);
 }
 
-void UAnimNextAnimationGraph_EditorData::OnCompileJobStarted()
+void UUAFAnimGraph_EditorData::OnCompileJobStarted()
 {
 	using namespace UE::UAF::UncookedOnly;
 
@@ -850,7 +850,7 @@ void UAnimNextAnimationGraph_EditorData::OnCompileJobStarted()
 	AnimationGraph->FreezeGraphInstances();
 }
 
-void UAnimNextAnimationGraph_EditorData::OnCompileJobFinished()
+void UUAFAnimGraph_EditorData::OnCompileJobFinished()
 {
 	using namespace UE::UAF::UncookedOnly;
 
@@ -860,7 +860,7 @@ void UAnimNextAnimationGraph_EditorData::OnCompileJobFinished()
 	AnimationGraph->ThawGraphInstances();
 }
 
-TConstArrayView<TSubclassOf<UAnimNextRigVMAssetEntry>> UAnimNextAnimationGraph_EditorData::GetEntryClasses() const
+TConstArrayView<TSubclassOf<UAnimNextRigVMAssetEntry>> UUAFAnimGraph_EditorData::GetEntryClasses() const
 {
 	static const TSubclassOf<UAnimNextRigVMAssetEntry> Classes[] =
 	{
@@ -872,7 +872,7 @@ TConstArrayView<TSubclassOf<UAnimNextRigVMAssetEntry>> UAnimNextAnimationGraph_E
 	return Classes;
 }
 
-bool UAnimNextAnimationGraph_EditorData::CanAddNewEntry(TSubclassOf<UAnimNextRigVMAssetEntry> InClass) const
+bool UUAFAnimGraph_EditorData::CanAddNewEntry(TSubclassOf<UAnimNextRigVMAssetEntry> InClass) const
 {
 	// Prevent users adding more than one animation graph
 	if(InClass == UAnimNextAnimationGraphEntry::StaticClass())
@@ -899,20 +899,20 @@ bool UAnimNextAnimationGraph_EditorData::CanAddNewEntry(TSubclassOf<UAnimNextRig
 
 UAnimNextAnimationGraphEntry* UAnimNextAnimationGraphLibrary::AddAnimationGraph(UUAFAnimGraph* InAsset, FName InName, bool bSetupUndoRedo, bool bPrintPythonCommand)
 {
-	return UE::UAF::UncookedOnly::FUtils::GetEditorData<UAnimNextAnimationGraph_EditorData>(InAsset)->AddAnimationGraph(InName, bSetupUndoRedo, bPrintPythonCommand);
+	return UE::UAF::UncookedOnly::FUtils::GetEditorData<UUAFAnimGraph_EditorData>(InAsset)->AddAnimationGraph(InName, bSetupUndoRedo, bPrintPythonCommand);
 }
 
-UAnimNextAnimationGraphEntry* UAnimNextAnimationGraph_EditorData::AddAnimationGraph(FName InName, bool bSetupUndoRedo, bool bPrintPythonCommand)
+UAnimNextAnimationGraphEntry* UUAFAnimGraph_EditorData::AddAnimationGraph(FName InName, bool bSetupUndoRedo, bool bPrintPythonCommand)
 {
 	if(InName == NAME_None)
 	{
-		ReportError(TEXT("UAnimNextRigVMAssetEditorData::AddAnimationGraph: Invalid graph name supplied."));
+		ReportError(TEXT("UUAFAnimGraph_EditorData::AddAnimationGraph: Invalid graph name supplied."));
 		return nullptr;
 	}
 
 	if(!GetEntryClasses().Contains(UAnimNextAnimationGraphEntry::StaticClass()) || !CanAddNewEntry(UAnimNextAnimationGraphEntry::StaticClass()))
 	{
-		ReportError(TEXT("UAnimNextRigVMAssetEditorData::AddAnimationGraph: Cannot add an animation graph to this asset - entry is not allowed."));
+		ReportError(TEXT("UUAFAnimGraph_EditorData::AddAnimationGraph: Cannot add an animation graph to this asset - entry is not allowed."));
 		return nullptr;
 	}
 
@@ -980,12 +980,12 @@ UAnimNextAnimationGraphEntry* UAnimNextAnimationGraph_EditorData::AddAnimationGr
 	return NewEntry;
 }
 
-TSubclassOf<UAssetUserData> UAnimNextAnimationGraph_EditorData::GetAssetUserDataClass() const
+TSubclassOf<UAssetUserData> UUAFAnimGraph_EditorData::GetAssetUserDataClass() const
 {
 	return UAnimNextAnimGraphWorkspaceAssetUserData::StaticClass();
 }
 
-void UAnimNextAnimationGraph_EditorData::InitializeAssetUserData()
+void UUAFAnimGraph_EditorData::InitializeAssetUserData()
 {
 	// Here we switch user data classes to patch up old assets
 	if (IInterface_AssetUserData* OuterUserData = Cast<IInterface_AssetUserData>(GetOuter()))
@@ -1003,12 +1003,12 @@ void UAnimNextAnimationGraph_EditorData::InitializeAssetUserData()
 	Super::InitializeAssetUserData();
 }
 
-UClass* UAnimNextAnimationGraph_EditorData::GetRigVMEdGraphSchemaClass() const
+UClass* UUAFAnimGraph_EditorData::GetRigVMEdGraphSchemaClass() const
 {
 	return UUAFAnimGraphEdGraphSchema::StaticClass();
 }
 
-void UAnimNextAnimationGraph_EditorData::HandleModifiedEvent(ERigVMGraphNotifType InNotifType, URigVMGraph* InGraph, UObject* InSubject)
+void UUAFAnimGraph_EditorData::HandleModifiedEvent(ERigVMGraphNotifType InNotifType, URigVMGraph* InGraph, UObject* InSubject)
 {
 	Super::HandleModifiedEvent(InNotifType, InGraph, InSubject);
 
diff --git a/Engine/Plugins/Experimental/UAF/UAFStateTree/Source/UAFStateTreeUncookedOnly/Internal/AnimNextStateTree_EditorData.h b/Engine/Plugins/Experimental/UAF/UAFStateTree/Source/UAFStateTreeUncookedOnly/Internal/AnimNextStateTree_EditorData.h
--- a/Engine/Plugins/Experimental/UAF/UAFStateTree/Source/UAFStateTreeUncookedOnly/Internal/AnimNextStateTree_EditorData.h
+++ b/Engine/Plugins/Experimental/UAF/UAFStateTree/Source/UAFStateTreeUncookedOnly/Internal/AnimNextStateTree_EditorData.h
@@ -7,7 +7,7 @@
 #include "AnimNextStateTree_EditorData.generated.h"
 
 UCLASS()
-class UAFSTATETREEUNCOOKEDONLY_API UAnimNextStateTree_EditorData : public UAnimNextAnimationGraph_EditorData
+class UAFSTATETREEUNCOOKEDONLY_API UAnimNextStateTree_EditorData : public UUAFAnimGraph_EditorData
 {
 	GENERATED_BODY()
 	
diff --git a/Engine/Plugins/Experimental/UAF/UAFStateTree/Source/UAFStateTreeUncookedOnly/Private/AnimNextStateTree_EditorData.cpp b/Engine/Plugins/Experimental/UAF/UAFStateTree/Source/UAFStateTreeUncookedOnly/Private/AnimNextStateTree_EditorData.cpp
--- a/Engine/Plugins/Experimental/UAF/UAFStateTree/Source/UAFStateTreeUncookedOnly/Private/AnimNextStateTree_EditorData.cpp
+++ b/Engine/Plugins/Experimental/UAF/UAFStateTree/Source/UAFStateTreeUncookedOnly/Private/AnimNextStateTree_EditorData.cpp
@@ -40,7 +40,7 @@ TSubclassOf<UAssetUserData> UAnimNextStateTree_EditorData::GetAssetUserDataClass
 
 void UAnimNextStateTree_EditorData::RecompileVM()
 {
-	UAnimNextAnimationGraph_EditorData::RecompileVM();
+	UUAFAnimGraph_EditorData::RecompileVM();
 
 	// Skip further state tree compilation if we had errors in regular asset compilation. It can create spurious errors
 	if (bErrorsDuringCompilation)
diff --git a/Engine/Plugins/Experimental/UAF/UAFTests/Source/UAFAnimGraphTestSuite/Private/AnimNextAnimGraphEditorVariablesTest.cpp b/Engine/Plugins/Experimental/UAF/UAFTests/Source/UAFAnimGraphTestSuite/Private/AnimNextAnimGraphEditorVariablesTest.cpp
--- a/Engine/Plugins/Experimental/UAF/UAFTests/Source/UAFAnimGraphTestSuite/Private/AnimNextAnimGraphEditorVariablesTest.cpp
+++ b/Engine/Plugins/Experimental/UAF/UAFTests/Source/UAFAnimGraphTestSuite/Private/AnimNextAnimGraphEditorVariablesTest.cpp
@@ -41,7 +41,7 @@ bool FEditor_AnimGraph_Variables::RunTest(const FString& InParameters)
 	
 	FFactoryAndClass FactoryClassPairs[] =
 	{
-		{ UAnimNextAnimationGraphFactory::StaticClass(), UUAFAnimGraph::StaticClass() },
+		{ UUAFAnimGraphFactory::StaticClass(), UUAFAnimGraph::StaticClass() },
 	};
 
 	for(const FFactoryAndClass& FactoryAndClass : FactoryClassPairs)
@@ -143,7 +143,7 @@ bool FEditor_AnimGraph_Graphs::RunTest(const FString& InParameters)
 	
 	FTestSettings TestSettings[] =
 	{
-		{ UAnimNextAnimationGraphFactory::StaticClass(), UUAFAnimGraph::StaticClass(), false, false, true, true },
+		{ UUAFAnimGraphFactory::StaticClass(), UUAFAnimGraph::StaticClass(), false, false, true, true },
 	};
 
 	for(const FTestSettings& TestSetting : TestSettings)
@@ -194,7 +194,7 @@ bool FEditor_AnimGraph_Graphs::RunTest(const FString& InParameters)
 		// AddAnimationGraph
 		if(TestSetting.bAnimGraphsAllowed)
 		{
-			UAnimNextAnimationGraph_EditorData* AnimationGraphEditorData = CastChecked<UAnimNextAnimationGraph_EditorData>(EditorData);
+			UUAFAnimGraph_EditorData* AnimationGraphEditorData = CastChecked<UUAFAnimGraph_EditorData>(EditorData);
 			
 			UAnimNextAnimationGraphEntry* AnimationGraphEntry = nullptr;
 			if(TestSetting.bExpectExistingAnimGraph)
diff --git a/Engine/Plugins/Experimental/UAF/UAFTests/Source/UAFAnimGraphTestSuite/Private/AnimNextAnimGraphTraitEventTest.cpp b/Engine/Plugins/Experimental/UAF/UAFTests/Source/UAFAnimGraphTestSuite/Private/AnimNextAnimGraphTraitEventTest.cpp
--- a/Engine/Plugins/Experimental/UAF/UAFTests/Source/UAFAnimGraphTestSuite/Private/AnimNextAnimGraphTraitEventTest.cpp
+++ b/Engine/Plugins/Experimental/UAF/UAFTests/Source/UAFAnimGraphTestSuite/Private/AnimNextAnimGraphTraitEventTest.cpp
@@ -194,7 +194,7 @@ bool FAnimationAnimNextRuntimeTest_GraphTraitEvent::RunTest(const FString& InPar
 		AUTO_REGISTER_ANIM_TRAIT(FTraitGraphTest_EventAB_OneChild)
 		AUTO_REGISTER_ANIM_TRAIT(FTraitGraphTest_EventAB_TwoChildren)
 
-		UFactory* GraphFactory = NewObject<UAnimNextAnimationGraphFactory>();
+		UFactory* GraphFactory = NewObject<UUAFAnimGraphFactory>();
 		UUAFAnimGraph* AnimationGraph = CastChecked<UUAFAnimGraph>(GraphFactory->FactoryCreateNew(UUAFAnimGraph::StaticClass(), GetTransientPackage(), TEXT("TestAnimNextGraph"), RF_Transient, nullptr, nullptr, NAME_None));
 		UE_RETURN_ON_ERROR(AnimationGraph != nullptr, "FAnimationAnimNextRuntimeTest_GraphTraitEvent -> Failed to create animation graph");
 
diff --git a/Engine/Plugins/Experimental/UAF/UAFTests/Source/UAFAnimGraphTestSuite/Private/AnimNextAnimGraphTraitGraphTest.cpp b/Engine/Plugins/Experimental/UAF/UAFTests/Source/UAFAnimGraphTestSuite/Private/AnimNextAnimGraphTraitGraphTest.cpp
--- a/Engine/Plugins/Experimental/UAF/UAFTests/Source/UAFAnimGraphTestSuite/Private/AnimNextAnimGraphTraitGraphTest.cpp
+++ b/Engine/Plugins/Experimental/UAF/UAFTests/Source/UAFAnimGraphTestSuite/Private/AnimNextAnimGraphTraitGraphTest.cpp
@@ -112,11 +112,11 @@ bool FAnimationAnimNextEditorTest_GraphAddTrait::RunTest(const FString& InParame
 
 		FScopedClearNodeTemplateRegistry ScopedClearNodeTemplateRegistry;
 
-		UFactory* GraphFactory = NewObject<UAnimNextAnimationGraphFactory>();
+		UFactory* GraphFactory = NewObject<UUAFAnimGraphFactory>();
 		UUAFAnimGraph* AnimationGraph = CastChecked<UUAFAnimGraph>(GraphFactory->FactoryCreateNew(UUAFAnimGraph::StaticClass(), GetTransientPackage(), TEXT("TestAnimNextGraph"), RF_Transient, nullptr, nullptr, NAME_None));
 		UE_RETURN_ON_ERROR(AnimationGraph != nullptr, "FAnimationAnimNextEditorTest_GraphAddTrait -> Failed to create animation graph");
 
-		UAnimNextAnimationGraph_EditorData* EditorData = UncookedOnly::FUtils::GetEditorData<UAnimNextAnimationGraph_EditorData>(AnimationGraph);
+		UUAFAnimGraph_EditorData* EditorData = UncookedOnly::FUtils::GetEditorData<UUAFAnimGraph_EditorData>(AnimationGraph);
 		UE_RETURN_ON_ERROR(EditorData != nullptr, "FAnimationAnimNextEditorTest_GraphAddTrait -> Failed to find module editor data");
 
 		URigVMController* Controller = EditorData->GetRigVMClient()->GetController(EditorData->GetRigVMClient()->GetDefaultModel());
@@ -233,11 +233,11 @@ bool FAnimationAnimNextEditorTest_GraphTraitOperations::RunTest(const FString& I
 
 		FScopedClearNodeTemplateRegistry ScopedClearNodeTemplateRegistry;
 
-		UFactory* GraphFactory = NewObject<UAnimNextAnimationGraphFactory>();
+		UFactory* GraphFactory = NewObject<UUAFAnimGraphFactory>();
 		UUAFAnimGraph* AnimationGraph = CastChecked<UUAFAnimGraph>(GraphFactory->FactoryCreateNew(UUAFAnimGraph::StaticClass(), GetTransientPackage(), TEXT("TestAnimNextGraph"), RF_Transient, nullptr, nullptr, NAME_None));
 		UE_RETURN_ON_ERROR(AnimationGraph != nullptr, "FAnimationAnimNextEditorTest_GraphTraitOperations -> Failed to create animation graph");
 
-		UAnimNextAnimationGraph_EditorData* EditorData = UncookedOnly::FUtils::GetEditorData<UAnimNextAnimationGraph_EditorData>(AnimationGraph);
+		UUAFAnimGraph_EditorData* EditorData = UncookedOnly::FUtils::GetEditorData<UUAFAnimGraph_EditorData>(AnimationGraph);
 		UE_RETURN_ON_ERROR(EditorData != nullptr, "FAnimationAnimNextEditorTest_GraphTraitOperations -> Failed to find module editor data");
 
 		UAnimNextController* Controller = Cast<UAnimNextController>(EditorData->GetRigVMClient()->GetController(EditorData->GetRigVMClient()->GetDefaultModel()));
@@ -415,11 +415,11 @@ bool FAnimationAnimNextRuntimeTest_GraphExecute::RunTest(const FString& InParame
 
 		FScopedClearNodeTemplateRegistry ScopedClearNodeTemplateRegistry;
 
-		UFactory* GraphFactory = NewObject<UAnimNextAnimationGraphFactory>();
+		UFactory* GraphFactory = NewObject<UUAFAnimGraphFactory>();
 		UUAFAnimGraph* AnimationGraph = CastChecked<UUAFAnimGraph>(GraphFactory->FactoryCreateNew(UUAFAnimGraph::StaticClass(), GetTransientPackage(), TEXT("TestAnimNextGraph"), RF_Transient, nullptr, nullptr, NAME_None));
 		UE_RETURN_ON_ERROR(AnimationGraph != nullptr, "FAnimationAnimNextEditorTest_GraphTraitOperations -> Failed to create animation graph");
 
-		UAnimNextAnimationGraph_EditorData* EditorData = UncookedOnly::FUtils::GetEditorData<UAnimNextAnimationGraph_EditorData>(AnimationGraph);
+		UUAFAnimGraph_EditorData* EditorData = UncookedOnly::FUtils::GetEditorData<UUAFAnimGraph_EditorData>(AnimationGraph);
 		UE_RETURN_ON_ERROR(EditorData != nullptr, "FAnimationAnimNextRuntimeTest_GraphExecute -> Failed to find module editor data");
 
 		URigVMController* Controller = EditorData->GetRigVMClient()->GetController(EditorData->GetRigVMClient()->GetDefaultModel());
@@ -529,11 +529,11 @@ bool FAnimationAnimNextRuntimeTest_GraphExecuteLatent::RunTest(const FString& In
 		AUTO_REGISTER_ANIM_TRAIT(FTestTrait)
 		FScopedClearNodeTemplateRegistry ScopedClearNodeTemplateRegistry;
 
-		UFactory* GraphFactory = NewObject<UAnimNextAnimationGraphFactory>();
+		UFactory* GraphFactory = NewObject<UUAFAnimGraphFactory>();
 		UUAFAnimGraph* AnimationGraph = CastChecked<UUAFAnimGraph>(GraphFactory->FactoryCreateNew(UUAFAnimGraph::StaticClass(), GetTransientPackage(), TEXT("TestAnimNextGraph"), RF_Transient, nullptr, nullptr, NAME_None));
 		UE_RETURN_ON_ERROR(AnimationGraph != nullptr, "FAnimationAnimNextEditorTest_GraphTraitOperations -> Failed to create animation graph");
 
-		UAnimNextAnimationGraph_EditorData* EditorData = UncookedOnly::FUtils::GetEditorData<UAnimNextAnimationGraph_EditorData>(AnimationGraph);
+		UUAFAnimGraph_EditorData* EditorData = UncookedOnly::FUtils::GetEditorData<UUAFAnimGraph_EditorData>(AnimationGraph);
 		UE_RETURN_ON_ERROR(EditorData != nullptr, "FAnimationAnimNextRuntimeTest_GraphExecuteLatent -> Failed to find module editor data");
 
 		EditorData->AddVariable(TEXT("TestIntVar"), FAnimNextParamType::GetType<int32>(), TEXT("34"));
@@ -674,7 +674,7 @@ bool FAnimationAnimNextRuntimeTest_Variables::RunTest(const FString& InParameter
 	using namespace UE::UAF;
 
 	{
-		UFactory* GraphFactory = NewObject<UAnimNextAnimationGraphFactory>();
+		UFactory* GraphFactory = NewObject<UUAFAnimGraphFactory>();
 		UUAFAnimGraph* AnimationGraph = CastChecked<UUAFAnimGraph>(GraphFactory->FactoryCreateNew(UUAFAnimGraph::StaticClass(), GetTransientPackage(), TEXT("TestAnimNextGraph"), RF_Transient, nullptr, nullptr, NAME_None));
 		UE_RETURN_ON_ERROR(AnimationGraph != nullptr, "FAnimationAnimNextRuntimeTest_Variables -> Failed to create animation graph");
 
diff --git a/Engine/Plugins/Experimental/UAF/UAFTests/Source/UAFAnimGraphTestSuite/Private/AnimNextAnimGraphTraitInterfacesTest.cpp b/Engine/Plugins/Experimental/UAF/UAFTests/Source/UAFAnimGraphTestSuite/Private/AnimNextAnimGraphTraitInterfacesTest.cpp
--- a/Engine/Plugins/Experimental/UAF/UAFTests/Source/UAFAnimGraphTestSuite/Private/AnimNextAnimGraphTraitInterfacesTest.cpp
+++ b/Engine/Plugins/Experimental/UAF/UAFTests/Source/UAFAnimGraphTestSuite/Private/AnimNextAnimGraphTraitInterfacesTest.cpp
@@ -428,7 +428,7 @@ bool FAnimationAnimNextRuntimeTest_IHierarchy::RunTest(const FString& InParamete
 		AUTO_REGISTER_ANIM_TRAIT(FTraitWithOneChild)
 		AUTO_REGISTER_ANIM_TRAIT(FTraitWithChildren)
 
-		UFactory* GraphFactory = NewObject<UAnimNextAnimationGraphFactory>();
+		UFactory* GraphFactory = NewObject<UUAFAnimGraphFactory>();
 		UUAFAnimGraph* AnimationGraph = CastChecked<UUAFAnimGraph>(GraphFactory->FactoryCreateNew(UUAFAnimGraph::StaticClass(), GetTransientPackage(), TEXT("TestAnimNextGraph"), RF_Transient, nullptr, nullptr, NAME_None));
 		UE_RETURN_ON_ERROR(AnimationGraph != nullptr, "FAnimationAnimNextRuntimeTest_GraphTraitEvent -> Failed to create animation graph");
 
@@ -643,7 +643,7 @@ bool FAnimationAnimNextRuntimeTest_IUpdate::RunTest(const FString& InParameters)
 		AUTO_REGISTER_ANIM_TRAIT(FTraitWithOneChild)
 		AUTO_REGISTER_ANIM_TRAIT(FTraitWithChildren)
 
-		UFactory* GraphFactory = NewObject<UAnimNextAnimationGraphFactory>();
+		UFactory* GraphFactory = NewObject<UUAFAnimGraphFactory>();
 		UUAFAnimGraph* AnimationGraph = CastChecked<UUAFAnimGraph>(GraphFactory->FactoryCreateNew(UUAFAnimGraph::StaticClass(), GetTransientPackage(), TEXT("TestAnimNextGraph"), RF_Transient, nullptr, nullptr, NAME_None));
 		UE_RETURN_ON_ERROR(AnimationGraph != nullptr, "FAnimationAnimNextRuntimeTest_GraphTraitEvent -> Failed to create animation graph");
 
@@ -779,7 +779,7 @@ bool FAnimationAnimNextRuntimeTest_IEvaluate::RunTest(const FString& InParameter
 		AUTO_REGISTER_ANIM_TRAIT(FTraitWithOneChild)
 		AUTO_REGISTER_ANIM_TRAIT(FTraitWithChildren)
 
-		UFactory* GraphFactory = NewObject<UAnimNextAnimationGraphFactory>();
+		UFactory* GraphFactory = NewObject<UUAFAnimGraphFactory>();
 		UUAFAnimGraph* AnimationGraph = CastChecked<UUAFAnimGraph>(GraphFactory->FactoryCreateNew(UUAFAnimGraph::StaticClass(), GetTransientPackage(), TEXT("TestAnimNextGraph"), RF_Transient, nullptr, nullptr, NAME_None));
 		UE_RETURN_ON_ERROR(AnimationGraph != nullptr, "FAnimationAnimNextRuntimeTest_GraphTraitEvent -> Failed to create animation graph");
 
@@ -914,7 +914,7 @@ bool FAnimationAnimNextRuntimeTest_IScopedInterface::RunTest(const FString& InPa
 		AUTO_REGISTER_ANIM_TRAIT(FScopedTagTrait)
 		AUTO_REGISTER_ANIM_TRAIT(FTestScopedTagTrait)
 
-		UFactory* GraphFactory = NewObject<UAnimNextAnimationGraphFactory>();
+		UFactory* GraphFactory = NewObject<UUAFAnimGraphFactory>();
 		UUAFAnimGraph* AnimationGraph = CastChecked<UUAFAnimGraph>(GraphFactory->FactoryCreateNew(UUAFAnimGraph::StaticClass(), GetTransientPackage(), TEXT("TestAnimNextGraph"), RF_Transient, nullptr, nullptr, NAME_None));
 		UE_RETURN_ON_ERROR(AnimationGraph != nullptr, "FAnimationAnimNextRuntimeTest_GraphTraitEvent -> Failed to create animation graph");
 
diff --git a/Engine/Plugins/Experimental/UAF/UAFTests/Source/UAFAnimGraphTestSuite/Private/AnimNextTraitBaseTest.cpp b/Engine/Plugins/Experimental/UAF/UAFTests/Source/UAFAnimGraphTestSuite/Private/AnimNextTraitBaseTest.cpp
--- a/Engine/Plugins/Experimental/UAF/UAFTests/Source/UAFAnimGraphTestSuite/Private/AnimNextTraitBaseTest.cpp
+++ b/Engine/Plugins/Experimental/UAF/UAFTests/Source/UAFAnimGraphTestSuite/Private/AnimNextTraitBaseTest.cpp
@@ -526,7 +526,7 @@ bool FAnimationAnimNextRuntimeTest_NodeLifetime::RunTest(const FString& InParame
 		AUTO_REGISTER_ANIM_TRAIT(FTraitAB_Add)
 		AUTO_REGISTER_ANIM_TRAIT(FTraitAC_Add)
 
-		UFactory* AnimationGraphFactory = NewObject<UAnimNextAnimationGraphFactory>();
+		UFactory* AnimationGraphFactory = NewObject<UUAFAnimGraphFactory>();
 		UUAFAnimGraph* AnimationGraph = CastChecked<UUAFAnimGraph>(AnimationGraphFactory->FactoryCreateNew(UUAFAnimGraph::StaticClass(), GetTransientPackage(), TEXT("TestAnimNextModule"), RF_Transient, nullptr, nullptr, NAME_None));
 		UE_RETURN_ON_ERROR(AnimationGraph != nullptr, "FAnimationAnimNextRuntimeTest_NodeLifetime -> Failed to create animation graph");
 
@@ -731,7 +731,7 @@ bool FAnimationAnimNextRuntimeTest_TraitStackBinding::RunTest(const FString& InP
 		AUTO_REGISTER_ANIM_TRAIT(FTraitAB_Add)
 		AUTO_REGISTER_ANIM_TRAIT(FTraitAC_Add)
 
-		UFactory* GraphFactory = NewObject<UAnimNextAnimationGraphFactory>();
+		UFactory* GraphFactory = NewObject<UUAFAnimGraphFactory>();
 		UUAFAnimGraph* AnimationGraph = CastChecked<UUAFAnimGraph>(GraphFactory->FactoryCreateNew(UUAFAnimGraph::StaticClass(), GetTransientPackage(), TEXT("TestAnimNextGraph"), RF_Transient, nullptr, nullptr, NAME_None));
 		UE_RETURN_ON_ERROR(AnimationGraph != nullptr, "FAnimationAnimNextRuntimeTest_TraitStackBinding -> Failed to create animation graph");
 
@@ -1151,7 +1151,7 @@ bool FAnimationAnimNextRuntimeTest_TraitSerialization::RunTest(const FString& In
 		AUTO_REGISTER_ANIM_TRAIT(FTraitSerialization_Add)
 		AUTO_REGISTER_ANIM_TRAIT(FTraitNativeSerialization_Add)
 
-		UFactory* GraphFactory = NewObject<UAnimNextAnimationGraphFactory>();
+		UFactory* GraphFactory = NewObject<UUAFAnimGraphFactory>();
 		UUAFAnimGraph* AnimationGraph = CastChecked<UUAFAnimGraph>(GraphFactory->FactoryCreateNew(UUAFAnimGraph::StaticClass(), GetTransientPackage(), TEXT("TestAnimNextGraph"), RF_Transient, nullptr, nullptr, NAME_None));
 		UE_RETURN_ON_ERROR(AnimationGraph != nullptr, "FAnimationAnimNextRuntimeTest_TraitSerialization -> Failed to create animation graph");
 
diff --git a/Engine/Plugins/Experimental/UAF/UAFTests/Source/UAFAnimGraphTestSuite/Private/AnimNextTraitEventTest.cpp b/Engine/Plugins/Experimental/UAF/UAFTests/Source/UAFAnimGraphTestSuite/Private/AnimNextTraitEventTest.cpp
--- a/Engine/Plugins/Experimental/UAF/UAFTests/Source/UAFAnimGraphTestSuite/Private/AnimNextTraitEventTest.cpp
+++ b/Engine/Plugins/Experimental/UAF/UAFTests/Source/UAFAnimGraphTestSuite/Private/AnimNextTraitEventTest.cpp
@@ -242,7 +242,7 @@ bool FAnimationAnimNextRuntimeTest_TraitEventRaising::RunTest(const FString& InP
 		AUTO_REGISTER_ANIM_TRAIT(FTraitCoreTest_EventA_Add)
 		AUTO_REGISTER_ANIM_TRAIT(FTraitCoreTest_EventB_Add)
 
-		UFactory* GraphFactory = NewObject<UAnimNextAnimationGraphFactory>();
+		UFactory* GraphFactory = NewObject<UUAFAnimGraphFactory>();
 		UUAFAnimGraph* AnimationGraph = CastChecked<UUAFAnimGraph>(GraphFactory->FactoryCreateNew(UUAFAnimGraph::StaticClass(), GetTransientPackage(), TEXT("TestAnimNextGraph"), RF_Transient, nullptr, nullptr, NAME_None));
 		UE_RETURN_ON_ERROR(AnimationGraph != nullptr, "FAnimationAnimNextEditorTest_GraphTraitOperations -> Failed to create animation graph");
 
diff --git a/Engine/Plugins/Experimental/UAF/UAFTests/Source/UAFAnimGraphTestSuite/Private/GCInstanceDataGCTest.cpp b/Engine/Plugins/Experimental/UAF/UAFTests/Source/UAFAnimGraphTestSuite/Private/GCInstanceDataGCTest.cpp
--- a/Engine/Plugins/Experimental/UAF/UAFTests/Source/UAFAnimGraphTestSuite/Private/GCInstanceDataGCTest.cpp
+++ b/Engine/Plugins/Experimental/UAF/UAFTests/Source/UAFAnimGraphTestSuite/Private/GCInstanceDataGCTest.cpp
@@ -93,11 +93,11 @@ bool FAnimationUAFRuntimeTest_InstanceDataGC::RunTest(const FString& InParameter
 		AUTO_REGISTER_ANIM_TRAIT(FTestTrait_InstanceDataGC)
 		FScopedClearNodeTemplateRegistry ScopedClearNodeTemplateRegistry;
 
-		UFactory* GraphFactory = NewObject<UAnimNextAnimationGraphFactory>();
+		UFactory* GraphFactory = NewObject<UUAFAnimGraphFactory>();
 		UUAFAnimGraph* AnimationGraph = CastChecked<UUAFAnimGraph>(GraphFactory->FactoryCreateNew(UUAFAnimGraph::StaticClass(), GetTransientPackage(), TEXT("TestAnimNextGraph"), RF_Transient, nullptr, nullptr, NAME_None));
 		UE_RETURN_ON_ERROR(AnimationGraph != nullptr, "FAnimationUAFRuntimeTest_InstanceDataGC -> Failed to create animation graph");
 
-		UAnimNextAnimationGraph_EditorData* EditorData = UncookedOnly::FUtils::GetEditorData<UAnimNextAnimationGraph_EditorData>(AnimationGraph);
+		UUAFAnimGraph_EditorData* EditorData = UncookedOnly::FUtils::GetEditorData<UUAFAnimGraph_EditorData>(AnimationGraph);
 		UE_RETURN_ON_ERROR(EditorData != nullptr, "FAnimationUAFRuntimeTest_InstanceDataGC -> Failed to find module editor data");
 
 		UAnimNextController* Controller = Cast<UAnimNextController>(EditorData->GetRigVMClient()->GetController(EditorData->GetRigVMClient()->GetDefaultModel()));
diff --git a/Engine/Plugins/Experimental/UAF/UAFTests/Source/UAFAnimGraphTestSuite/Private/GCInstanceDataMissingGCTest.cpp b/Engine/Plugins/Experimental/UAF/UAFTests/Source/UAFAnimGraphTestSuite/Private/GCInstanceDataMissingGCTest.cpp
--- a/Engine/Plugins/Experimental/UAF/UAFTests/Source/UAFAnimGraphTestSuite/Private/GCInstanceDataMissingGCTest.cpp
+++ b/Engine/Plugins/Experimental/UAF/UAFTests/Source/UAFAnimGraphTestSuite/Private/GCInstanceDataMissingGCTest.cpp
@@ -78,11 +78,11 @@ bool FAnimationUAFRuntimeTest_InstanceDataGCMissing::RunTest(const FString& InPa
 		AUTO_REGISTER_ANIM_TRAIT(FTestTrait_InstanceDataGCMissing)
 		FScopedClearNodeTemplateRegistry ScopedClearNodeTemplateRegistry;
 
-		UFactory* GraphFactory = NewObject<UAnimNextAnimationGraphFactory>();
+		UFactory* GraphFactory = NewObject<UUAFAnimGraphFactory>();
 		UUAFAnimGraph* AnimationGraph = CastChecked<UUAFAnimGraph>(GraphFactory->FactoryCreateNew(UUAFAnimGraph::StaticClass(), GetTransientPackage(), TEXT("TestAnimNextGraph"), RF_Transient, nullptr, nullptr, NAME_None));
 		UE_RETURN_ON_ERROR(AnimationGraph != nullptr, "FAnimationUAFRuntimeTest_InstanceDataGCMissing -> Failed to create animation graph");
 
-		UAnimNextAnimationGraph_EditorData* EditorData = UncookedOnly::FUtils::GetEditorData<UAnimNextAnimationGraph_EditorData>(AnimationGraph);
+		UUAFAnimGraph_EditorData* EditorData = UncookedOnly::FUtils::GetEditorData<UUAFAnimGraph_EditorData>(AnimationGraph);
 		UE_RETURN_ON_ERROR(EditorData != nullptr, "FAnimationUAFRuntimeTest_InstanceDataGCMissing -> Failed to find module editor data");
 
 		UAnimNextController* Controller = Cast<UAnimNextController>(EditorData->GetRigVMClient()->GetController(EditorData->GetRigVMClient()->GetDefaultModel()));
diff --git a/Engine/Plugins/Experimental/UAF/UAFTests/Source/UAFAnimGraphTestSuite/Private/GCSharedDataGCTest.cpp b/Engine/Plugins/Experimental/UAF/UAFTests/Source/UAFAnimGraphTestSuite/Private/GCSharedDataGCTest.cpp
--- a/Engine/Plugins/Experimental/UAF/UAFTests/Source/UAFAnimGraphTestSuite/Private/GCSharedDataGCTest.cpp
+++ b/Engine/Plugins/Experimental/UAF/UAFTests/Source/UAFAnimGraphTestSuite/Private/GCSharedDataGCTest.cpp
@@ -97,11 +97,11 @@ bool FAnimationUAFRuntimeTest_SharedDataGC::RunTest(const FString& InParameters)
 		AUTO_REGISTER_ANIM_TRAIT(FTestTrait_SharedDataGC)
 		FScopedClearNodeTemplateRegistry ScopedClearNodeTemplateRegistry;
 
-		UFactory* GraphFactory = NewObject<UAnimNextAnimationGraphFactory>();
+		UFactory* GraphFactory = NewObject<UUAFAnimGraphFactory>();
 		TObjectPtr<UUAFAnimGraph> AnimationGraph = CastChecked<UUAFAnimGraph>(GraphFactory->FactoryCreateNew(UUAFAnimGraph::StaticClass(), GetTransientPackage(), TEXT("TestAnimNextGraph"), RF_Transient, nullptr, nullptr, NAME_None));
 		UE_RETURN_ON_ERROR(AnimationGraph != nullptr, "FAnimationUAFRuntimeTest_TraitGC -> Failed to create animation graph");
 
-		TObjectPtr<UAnimNextAnimationGraph_EditorData> EditorData = UncookedOnly::FUtils::GetEditorData<UAnimNextAnimationGraph_EditorData>(AnimationGraph.Get());
+		TObjectPtr<UUAFAnimGraph_EditorData> EditorData = UncookedOnly::FUtils::GetEditorData<UUAFAnimGraph_EditorData>(AnimationGraph.Get());
 		UE_RETURN_ON_ERROR(EditorData != nullptr, "FAnimationUAFRuntimeTest_TraitGC -> Failed to find module editor data");
 
 		TObjectPtr<UAnimNextController> Controller = Cast<UAnimNextController>(EditorData->GetRigVMClient()->GetController(EditorData->GetRigVMClient()->GetDefaultModel()));
diff --git a/Engine/Plugins/Experimental/UAF/UAFTests/Source/UAFAnimGraphTestSuite/Private/GCSharedDataMissingGCTest.cpp b/Engine/Plugins/Experimental/UAF/UAFTests/Source/UAFAnimGraphTestSuite/Private/GCSharedDataMissingGCTest.cpp
--- a/Engine/Plugins/Experimental/UAF/UAFTests/Source/UAFAnimGraphTestSuite/Private/GCSharedDataMissingGCTest.cpp
+++ b/Engine/Plugins/Experimental/UAF/UAFTests/Source/UAFAnimGraphTestSuite/Private/GCSharedDataMissingGCTest.cpp
@@ -51,11 +51,11 @@ bool FAnimationUAFRuntimeTest_SharedDataMissingGC::RunTest(const FString& InPara
 
 		AddExpectedError(TEXT("Trait is connected to a UObject but does not implement IGarbageCollection"), EAutomationExpectedErrorFlags::Contains, 6);
 
-		UFactory* GraphFactory = NewObject<UAnimNextAnimationGraphFactory>();
+		UFactory* GraphFactory = NewObject<UUAFAnimGraphFactory>();
 		UUAFAnimGraph* AnimationGraph = CastChecked<UUAFAnimGraph>(GraphFactory->FactoryCreateNew(UUAFAnimGraph::StaticClass(), GetTransientPackage(), TEXT("TestAnimNextGraph"), RF_Transient, nullptr, nullptr, NAME_None));
 		UE_RETURN_ON_ERROR(AnimationGraph != nullptr, "FAnimationUAFRuntimeTest_TraitMissingGC -> Failed to create animation graph");
 
-		UAnimNextAnimationGraph_EditorData* EditorData = UncookedOnly::FUtils::GetEditorData<UAnimNextAnimationGraph_EditorData>(AnimationGraph);
+		UUAFAnimGraph_EditorData* EditorData = UncookedOnly::FUtils::GetEditorData<UUAFAnimGraph_EditorData>(AnimationGraph);
 		UE_RETURN_ON_ERROR(EditorData != nullptr, "FAnimationUAFRuntimeTest_TraitMissingGC -> Failed to find module editor data");
 
 		UAnimNextController* Controller = Cast<UAnimNextController>(EditorData->GetRigVMClient()->GetController(EditorData->GetRigVMClient()->GetDefaultModel()));
diff --git a/Engine/Plugins/Experimental/UAF/UAFTests/Source/UAFTests/Private/UAFAnimationGraphTests.cpp b/Engine/Plugins/Experimental/UAF/UAFTests/Source/UAFTests/Private/UAFAnimationGraphTests.cpp
--- a/Engine/Plugins/Experimental/UAF/UAFTests/Source/UAFTests/Private/UAFAnimationGraphTests.cpp
+++ b/Engine/Plugins/Experimental/UAF/UAFTests/Source/UAFTests/Private/UAFAnimationGraphTests.cpp
@@ -167,8 +167,8 @@ TEST_CLASS(FAnimNextAnimationGraphTests, "Animation.UAF.Functional")
 protected:
 	void CreateAnimNextAnimationGraph()
 	{
-		const TSubclassOf<UAnimNextAnimationGraphFactory> AnimGraphFactoryClass = UAnimNextAnimationGraphFactory::StaticClass();
-		UObject* FactoryObject = UAFTestsUtilities::CreateFactoryObject(NewObject<UAnimNextAnimationGraphFactory>(GetTransientPackage(), AnimGraphFactoryClass.Get()), UUAFAnimGraph::StaticClass(), "NewAnimNextAnimationGraph");
+		const TSubclassOf<UUAFAnimGraphFactory> AnimGraphFactoryClass = UUAFAnimGraphFactory::StaticClass();
+		UObject* FactoryObject = UAFTestsUtilities::CreateFactoryObject(NewObject<UUAFAnimGraphFactory>(GetTransientPackage(), AnimGraphFactoryClass.Get()), UUAFAnimGraph::StaticClass(), "NewAnimNextAnimationGraph");
 		AnimNextAnimationGraph = CastChecked<UUAFAnimGraph>(FactoryObject);
 		ASSERT_THAT(IsNotNull(AnimNextAnimationGraph));
 		ASSERT_THAT(IsTrue(AssetEditorSubsystem->OpenEditorForAsset(AnimNextAnimationGraph)));
diff --git a/Engine/Plugins/Experimental/UAF/UAFTests/Source/UAFTests/Private/UAFAssetCreationTests.cpp b/Engine/Plugins/Experimental/UAF/UAFTests/Source/UAFTests/Private/UAFAssetCreationTests.cpp
--- a/Engine/Plugins/Experimental/UAF/UAFTests/Source/UAFTests/Private/UAFAssetCreationTests.cpp
+++ b/Engine/Plugins/Experimental/UAF/UAFTests/Source/UAFTests/Private/UAFAssetCreationTests.cpp
@@ -49,8 +49,8 @@ TEST_CLASS(FAssetCreationTests, "Animation.UAF.Functional")
 		TestCommandBuilder
 			.Do(TEXT("Create AnimNext Animation Graph Asset"), [&]()
 			{
-				const TSubclassOf<UAnimNextAnimationGraphFactory> AnimGraphFactoryClass = UAnimNextAnimationGraphFactory::StaticClass();
-				UObject* FactoryObject = UAFTestsUtilities::CreateFactoryObject(NewObject<UAnimNextAnimationGraphFactory>(), UUAFAnimGraph::StaticClass(), "NewAnimNextAnimationGraph");
+				const TSubclassOf<UUAFAnimGraphFactory> AnimGraphFactoryClass = UUAFAnimGraphFactory::StaticClass();
+				UObject* FactoryObject = UAFTestsUtilities::CreateFactoryObject(NewObject<UUAFAnimGraphFactory>(), UUAFAnimGraph::StaticClass(), "NewAnimNextAnimationGraph");
 				AnimNextAsset = CastChecked<UUAFAnimGraph>(FactoryObject);
 			})
 			.Then(TEXT("Assert AnimNext Animation Graph Asset Created and Open in Editor"), [&]()
diff --git a/Engine/Plugins/Experimental/UAF/UAFTests/Source/UAFTests/Private/UAFFunctionNodeTests.cpp b/Engine/Plugins/Experimental/UAF/UAFTests/Source/UAFTests/Private/UAFFunctionNodeTests.cpp
--- a/Engine/Plugins/Experimental/UAF/UAFTests/Source/UAFTests/Private/UAFFunctionNodeTests.cpp
+++ b/Engine/Plugins/Experimental/UAF/UAFTests/Source/UAFTests/Private/UAFFunctionNodeTests.cpp
@@ -36,8 +36,8 @@ TEST_CLASS(FAssetFunctionNodeTests, "Animation.UAF.Functional")
 		TestCommandBuilder
 			.Do(TEXT("Create AnimNext Animation Graph Asset"), [&]()
 			{
-				const TSubclassOf<UAnimNextAnimationGraphFactory> AnimNextAnimGraphFactoryClass = UAnimNextAnimationGraphFactory::StaticClass();
-				UObject* factoryObject = UAFTestsUtilities::CreateFactoryObject(NewObject<UAnimNextAnimationGraphFactory>(GetTransientPackage(), AnimNextAnimGraphFactoryClass.Get()), UUAFAnimGraph::StaticClass(), "NewAnimNextAnimationGraph");
+				const TSubclassOf<UUAFAnimGraphFactory> AnimNextAnimGraphFactoryClass = UUAFAnimGraphFactory::StaticClass();
+				UObject* factoryObject = UAFTestsUtilities::CreateFactoryObject(NewObject<UUAFAnimGraphFactory>(GetTransientPackage(), AnimNextAnimGraphFactoryClass.Get()), UUAFAnimGraph::StaticClass(), "NewAnimNextAnimationGraph");
 				AnimNextAsset = CastChecked<UUAFAnimGraph>(factoryObject);				
 				ASSERT_THAT(IsTrue(AssetEditorSubsystem->OpenEditorForAsset(AnimNextAsset)));
 			})
@@ -115,8 +115,8 @@ TEST_CLASS(FAssetFunctionNodeTests, "Animation.UAF.Functional")
 		TestCommandBuilder
 			.Do(TEXT("Create AnimNext Animation Graph Asset"), [&]()
 			{
-				const TSubclassOf<UAnimNextAnimationGraphFactory> AnimNextAnimGraphFactoryClass = UAnimNextAnimationGraphFactory::StaticClass();
-				UObject* factoryObject = UAFTestsUtilities::CreateFactoryObject(NewObject<UAnimNextAnimationGraphFactory>(GetTransientPackage(), AnimNextAnimGraphFactoryClass.Get()), UUAFAnimGraph::StaticClass(), "NewAnimNextAnimationGraph");
+				const TSubclassOf<UUAFAnimGraphFactory> AnimNextAnimGraphFactoryClass = UUAFAnimGraphFactory::StaticClass();
+				UObject* factoryObject = UAFTestsUtilities::CreateFactoryObject(NewObject<UUAFAnimGraphFactory>(GetTransientPackage(), AnimNextAnimGraphFactoryClass.Get()), UUAFAnimGraph::StaticClass(), "NewAnimNextAnimationGraph");
 				AnimNextAsset = CastChecked<UUAFAnimGraph>(factoryObject);				
 				ASSERT_THAT(IsTrue(AssetEditorSubsystem->OpenEditorForAsset(AnimNextAsset)));
 			})
diff --git a/Engine/Plugins/Experimental/UAF/UAFTests/Source/UAFTests/Private/UAFVariableTests.cpp b/Engine/Plugins/Experimental/UAF/UAFTests/Source/UAFTests/Private/UAFVariableTests.cpp
--- a/Engine/Plugins/Experimental/UAF/UAFTests/Source/UAFTests/Private/UAFVariableTests.cpp
+++ b/Engine/Plugins/Experimental/UAF/UAFTests/Source/UAFTests/Private/UAFVariableTests.cpp
@@ -29,8 +29,8 @@ TEST_CLASS(FAssetVariableTests, "Animation.UAF.Functional")
 		TestCommandBuilder
 			.Do(TEXT("Create AnimNext Animation Graph Asset"), [&]()
 			{
-				const TSubclassOf<UAnimNextAnimationGraphFactory> AnimNextAnimGraphFactoryClass = UAnimNextAnimationGraphFactory::StaticClass();
-				UObject* factoryObject = UAFTestsUtilities::CreateFactoryObject(NewObject<UAnimNextAnimationGraphFactory>(GetTransientPackage(), AnimNextAnimGraphFactoryClass.Get()), UUAFAnimGraph::StaticClass(), "NewAnimNextAnimationGraph");
+				const TSubclassOf<UUAFAnimGraphFactory> AnimNextAnimGraphFactoryClass = UUAFAnimGraphFactory::StaticClass();
+				UObject* factoryObject = UAFTestsUtilities::CreateFactoryObject(NewObject<UUAFAnimGraphFactory>(GetTransientPackage(), AnimNextAnimGraphFactoryClass.Get()), UUAFAnimGraph::StaticClass(), "NewAnimNextAnimationGraph");
 				AnimNextAsset = CastChecked<UUAFAnimGraph>(factoryObject);				
 				ASSERT_THAT(IsTrue(AssetEditorSubsystem->OpenEditorForAsset(AnimNextAsset)));
 			})
diff --git a/Engine/Plugins/Experimental/UAF/UAFTests/Source/UAFTests/Private/UAFWorkspaceTests.cpp b/Engine/Plugins/Experimental/UAF/UAFTests/Source/UAFTests/Private/UAFWorkspaceTests.cpp
--- a/Engine/Plugins/Experimental/UAF/UAFTests/Source/UAFTests/Private/UAFWorkspaceTests.cpp
+++ b/Engine/Plugins/Experimental/UAF/UAFTests/Source/UAFTests/Private/UAFWorkspaceTests.cpp
@@ -33,8 +33,8 @@ TEST_CLASS(FUAFWorkspaceTests, "Animation.UAF.Functional")
 		TestCommandBuilder
 			.Do(TEXT("Create Animation Graph Asset"), [&]()
 			{
-				const TSubclassOf<UAnimNextAnimationGraphFactory> AnimGraphFactoryClass = UAnimNextAnimationGraphFactory::StaticClass();
-				UObject* FactoryObject = UAFTestsUtilities::CreateFactoryObject(NewObject<UAnimNextAnimationGraphFactory>(), UUAFAnimGraph::StaticClass(), "NewAnimNextAnimationGraph");
+				const TSubclassOf<UUAFAnimGraphFactory> AnimGraphFactoryClass = UUAFAnimGraphFactory::StaticClass();
+				UObject* FactoryObject = UAFTestsUtilities::CreateFactoryObject(NewObject<UUAFAnimGraphFactory>(), UUAFAnimGraph::StaticClass(), "NewAnimNextAnimationGraph");
 				AnimationGraphAsset = CastChecked<UUAFAnimGraph>(FactoryObject);
 				ASSERT_THAT(IsNotNull(AnimationGraphAsset));
 			})
@@ -50,8 +50,8 @@ TEST_CLASS(FUAFWorkspaceTests, "Animation.UAF.Functional")
 		TestCommandBuilder
 			.Do(TEXT("Create Animation Graph Asset"), [&]()
 			{
-				const TSubclassOf<UAnimNextAnimationGraphFactory> AnimGraphFactoryClass = UAnimNextAnimationGraphFactory::StaticClass();
-				UObject* FactoryObject = UAFTestsUtilities::CreateFactoryObject(NewObject<UAnimNextAnimationGraphFactory>(), UUAFAnimGraph::StaticClass(), "NewAnimNextAnimationGraph");
+				const TSubclassOf<UUAFAnimGraphFactory> AnimGraphFactoryClass = UUAFAnimGraphFactory::StaticClass();
+				UObject* FactoryObject = UAFTestsUtilities::CreateFactoryObject(NewObject<UUAFAnimGraphFactory>(), UUAFAnimGraph::StaticClass(), "NewAnimNextAnimationGraph");
 				AnimationGraphAsset = CastChecked<UUAFAnimGraph>(FactoryObject);
 				ASSERT_THAT(IsNotNull(AnimationGraphAsset));
 			})
@@ -76,12 +76,12 @@ TEST_CLASS(FUAFWorkspaceTests, "Animation.UAF.Functional")
 		TestCommandBuilder
 			.Do(TEXT("Create Animation Graph Assets"), [&]()
 			{
-				const TSubclassOf<UAnimNextAnimationGraphFactory> AnimGraphFactoryClass = UAnimNextAnimationGraphFactory::StaticClass();
-				UObject* FactoryObject = UAFTestsUtilities::CreateFactoryObject(NewObject<UAnimNextAnimationGraphFactory>(), UUAFAnimGraph::StaticClass(), "NewAnimNextAnimationGraph");
+				const TSubclassOf<UUAFAnimGraphFactory> AnimGraphFactoryClass = UUAFAnimGraphFactory::StaticClass();
+				UObject* FactoryObject = UAFTestsUtilities::CreateFactoryObject(NewObject<UUAFAnimGraphFactory>(), UUAFAnimGraph::StaticClass(), "NewAnimNextAnimationGraph");
 				AnimationGraphAsset = CastChecked<UUAFAnimGraph>(FactoryObject);
 				ASSERT_THAT(IsNotNull(AnimationGraphAsset));
 
-				UObject* FactoryObject2 = UAFTestsUtilities::CreateFactoryObject(NewObject<UAnimNextAnimationGraphFactory>(), UUAFAnimGraph::StaticClass(), "NewAnimNextAnimationGraph");
+				UObject* FactoryObject2 = UAFTestsUtilities::CreateFactoryObject(NewObject<UUAFAnimGraphFactory>(), UUAFAnimGraph::StaticClass(), "NewAnimNextAnimationGraph");
 				AnimationGraphAsset2 = CastChecked<UUAFAnimGraph>(FactoryObject2);
 				ASSERT_THAT(IsNotNull(AnimationGraphAsset2));
 			})
