[UAF][ANIM] Allow Option To Supply AnimInstance Class For Legacy Montages
- Some Gameplay Assumes Instance Class Types for Callback Hooks When completing Montage. Being able to specify the AnimInstance to use allows for Greater Control when Playing Legacy Montages

#Synced-CL 49653144

diff --git a/Engine/Plugins/Experimental/UAF/UAFAnimGraph/Source/UAFAnimGraph/Internal/Injection/UAFMontageComponent.h b/Engine/Plugins/Experimental/UAF/UAFAnimGraph/Source/UAFAnimGraph/Internal/Injection/UAFMontageComponent.h
--- a/Engine/Plugins/Experimental/UAF/UAFAnimGraph/Source/UAFAnimGraph/Internal/Injection/UAFMontageComponent.h
+++ b/Engine/Plugins/Experimental/UAF/UAFAnimGraph/Source/UAFAnimGraph/Internal/Injection/UAFMontageComponent.h
@@ -99,6 +99,10 @@ private:
 	UPROPERTY(EditAnywhere, Category = Montages, meta = (CustomWidget = AnimNextModuleEvent))
 	FName TickMontagesBeforeEventName = NAME_None;
 
+	// The Instance Class that Montages will run on
+	UPROPERTY(EditAnywhere, Category= Montages, DisplayName="Legacy Montage AnimInstance")
+	TSubclassOf<UAnimInstance> LegacyMontageAnimInstanceClass = nullptr;
+
 	// The event name of the module tick function used at runtime to enqueue the montage tick function as a prerequisite to. 
 	// If the the exposed TickMontagesBeforeEventName does not provide an override, it will fall back to DefaultTickDependencyEventName
 	UPROPERTY(Transient)
diff --git a/Engine/Plugins/Experimental/UAF/UAFAnimGraph/Source/UAFAnimGraph/Private/Injection/UAFMontageComponent.cpp b/Engine/Plugins/Experimental/UAF/UAFAnimGraph/Source/UAFAnimGraph/Private/Injection/UAFMontageComponent.cpp
--- a/Engine/Plugins/Experimental/UAF/UAFAnimGraph/Source/UAFAnimGraph/Private/Injection/UAFMontageComponent.cpp
+++ b/Engine/Plugins/Experimental/UAF/UAFAnimGraph/Source/UAFAnimGraph/Private/Injection/UAFMontageComponent.cpp
@@ -187,10 +187,17 @@ void FUAFMontageComponent::OnBindToInstance()
 			UAnimInstance* UsedAnimInstance = SkelMeshComponent->GetAnimInstance();
 			if (UsedAnimInstance == nullptr)
 			{
+
+				UClass* MontageInstanceClass = MontageComponent->LegacyMontageAnimInstanceClass;
+				if (MontageInstanceClass == nullptr)
+				{
+					MontageInstanceClass= UAnimInstance::StaticClass();
+				}
+
 				// Create a new AnimInstance that we can use to update our montages. 
-				SkelMeshComponent->AnimClass = UAnimInstance::StaticClass();
+				SkelMeshComponent->AnimClass = MontageInstanceClass;
 				SkelMeshComponent->ClearAnimScriptInstance();
-				SkelMeshComponent->AnimScriptInstance = NewObject<UAnimInstance>(SkelMeshComponent, UAnimInstance::StaticClass());
+				SkelMeshComponent->AnimScriptInstance = NewObject<UAnimInstance>(SkelMeshComponent, MontageInstanceClass);
 
 				UsedAnimInstance = SkelMeshComponent->GetAnimInstance();
 			}
