UAF: Add button to browse to trait shared data source file

#Synced-CL 49723230

diff --git a/Engine/Plugins/Experimental/UAF/UAFAnimGraph/Source/UAFAnimGraphEditor/Private/TraitStackEditor.cpp b/Engine/Plugins/Experimental/UAF/UAFAnimGraph/Source/UAFAnimGraphEditor/Private/TraitStackEditor.cpp
--- a/Engine/Plugins/Experimental/UAF/UAFAnimGraph/Source/UAFAnimGraphEditor/Private/TraitStackEditor.cpp
+++ b/Engine/Plugins/Experimental/UAF/UAFAnimGraph/Source/UAFAnimGraphEditor/Private/TraitStackEditor.cpp
@@ -9,6 +9,7 @@
 #include "IDetailPropertyRow.h"
 #include "IWorkspaceEditor.h"
 #include "PropertyHandle.h"
+#include "SourceCodeNavigation.h"
 #include "SPositiveActionButton.h"
 #include "SSimpleComboButton.h"
 #include "StructViewerFilter.h"
@@ -520,7 +521,43 @@ namespace Private
 					HandleWidget
 				]
 				+SHorizontalBox::Slot()
-				.FillWidth(1.0f)
+				.HAlign(HAlign_Right)
+				.FillContentWidth(1.0f, 0.0f) 
+				[
+					SNew(SBox)
+					.HAlign(HAlign_Right)
+					.VAlign(VAlign_Center)
+					.Padding(FMargin(4.f, 0.f, 0.f, 0.f))
+					.Visibility_Lambda([this]() -> EVisibility
+						{
+							if (FSourceCodeNavigation::CanNavigateToStruct(SharedDataStruct))
+							{
+								return EVisibility::Visible;
+							}
+							return EVisibility::Collapsed;
+						})
+					[
+						SNew(SButton)
+						.ButtonStyle(FAppStyle::Get(), "SimpleButton")
+						.OnClicked_Lambda([this]()
+							{
+								if (SharedDataStruct)
+								{
+									FSourceCodeNavigation::NavigateToStruct(SharedDataStruct);
+								}
+								return FReply::Handled();
+							})
+						.ToolTipText(FText::Format(LOCTEXT("GoToCode_ToolTip", "Click to open the trait shared data source file in {0}"), FSourceCodeNavigation::GetSelectedSourceCodeIDE()))
+						.ContentPadding(0.0f)
+						[
+							SNew(SImage)
+							.Image(FAppStyle::GetBrush("Icons.OpenSourceLocation"))
+							.ColorAndOpacity(FSlateColor::UseForeground())
+						]
+					]
+				]
+				+SHorizontalBox::Slot()
+				.AutoWidth()
 				.HAlign(HAlign_Right)
 				[
 					SNew(SSimpleComboButton)
