UAF: Support for latent parameters in modify curve trait

#Synced-CL 49767845

diff --git a/Engine/Plugins/Experimental/UAF/UAFAnimGraph/Source/UAFAnimGraph/Private/Traits/ModifyCurveTrait.cpp b/Engine/Plugins/Experimental/UAF/UAFAnimGraph/Source/UAFAnimGraph/Private/Traits/ModifyCurveTrait.cpp
--- a/Engine/Plugins/Experimental/UAF/UAFAnimGraph/Source/UAFAnimGraph/Private/Traits/ModifyCurveTrait.cpp
+++ b/Engine/Plugins/Experimental/UAF/UAFAnimGraph/Source/UAFAnimGraph/Private/Traits/ModifyCurveTrait.cpp
@@ -28,22 +28,23 @@ void UE::UAF::FModifyCurveTrait::PostEvaluate(FEvaluateTraversalContext& Context
 	FInstanceData* InstanceData = Binding.GetInstanceData<FInstanceData>();
 	check(InstanceData);
 
-	// Copy Alpha value since it is 'latent'
-	InstanceData->Alpha = SharedData->GetAlpha(Binding);
+	float Alpha = SharedData->GetAlpha(Binding);
+	EAnimNext_ModifyCurveApplyMode ApplyMode = SharedData->GetApplyMode(Binding);
+	TConstArrayView<FModifyCurveParameters> CurveParameters = SharedData->GetModifyCurveParameters(Binding);;
 
 #if ENABLE_ANIM_DEBUG 
 	InstanceData->HostObject = Context.GetHostObject();
 #endif // ENABLE_ANIM_DEBUG 
 
-	Context.AppendTask(FModifyCurveTask::Make(InstanceData, SharedData));
+	Context.AppendTask(FModifyCurveTask::Make(Alpha, ApplyMode, CurveParameters));
 }
 
-FModifyCurveTask FModifyCurveTask::Make(UE::UAF::FModifyCurveTrait::FInstanceData* InstanceData,
-	const UE::UAF::FModifyCurveTrait::FSharedData* SharedData)
+FModifyCurveTask FModifyCurveTask::Make(float Alpha, EAnimNext_ModifyCurveApplyMode ApplyMode, TConstArrayView<FModifyCurveParameters> CurveParameters)
 {
 	FModifyCurveTask Task;
-	Task.InstanceData = InstanceData;
-	Task.SharedData = SharedData;
+	Task.Alpha = Alpha;
+	Task.ApplyMode = ApplyMode;
+	Task.ModifyCurveParameters = CurveParameters;
 	return Task;
 }
 
@@ -62,12 +63,12 @@ void FModifyCurveTask::Execute(UE::UAF::FEvaluationVM& VM) const
 				
 				if (TUnboundValueMap<FFloatAnimationAttribute>* FloatMap = (*Collection)->GetUnboundValueMaps().Find<FFloatAnimationAttribute>())
 				{
-					for (const FModifyCurveParameters& CurveParameters : SharedData->ModifyCurveParameters)
+					for (const FModifyCurveParameters& CurveParameters : ModifyCurveParameters)
 					{
 						if (FFloatAnimationAttribute* Float = FloatMap->Find(CurveParameters.CurveName))
 						{
 							const float CurrentValue = (*Float).Value;
-							const float UpdatedValue = ProcessCurveOperation(CurrentValue, CurveParameters.CurveValue, InstanceData->Alpha, SharedData->ApplyMode);
+							const float UpdatedValue = ProcessCurveOperation(CurrentValue, CurveParameters.CurveValue, Alpha, ApplyMode);
 							(*Float).Value = UpdatedValue;
 						}
 					}
@@ -75,12 +76,12 @@ void FModifyCurveTask::Execute(UE::UAF::FEvaluationVM& VM) const
 
 				if (TBoundValueMap<FFloatAnimationAttribute>* FloatMap = (*Collection)->GetBoundValueMaps().Find<FFloatAnimationAttribute>(MappingKey))
 				{
-					for (const FModifyCurveParameters& CurveParameters : SharedData->ModifyCurveParameters)
+					for (const FModifyCurveParameters& CurveParameters : ModifyCurveParameters)
 					{
 						if (const FAttributeSetIndex RootMotionAttributeIndex = FloatTypedSet->FindIndex(CurveParameters.CurveName))
 						{
 							const float CurrentValue = (*FloatMap)[RootMotionAttributeIndex].Value;
-							const float UpdatedValue = ProcessCurveOperation(CurrentValue, CurveParameters.CurveValue, InstanceData->Alpha, SharedData->ApplyMode);
+							const float UpdatedValue = ProcessCurveOperation(CurrentValue, CurveParameters.CurveValue, Alpha, ApplyMode);
 							(*FloatMap)[RootMotionAttributeIndex].Value = UpdatedValue;
 						}
 					}
@@ -91,10 +92,10 @@ void FModifyCurveTask::Execute(UE::UAF::FEvaluationVM& VM) const
 	
 	if (const TUniquePtr<FKeyframeState>* Keyframe = VM.PeekValue<TUniquePtr<FKeyframeState>>(KEYFRAME_STACK_NAME, 0))
 	{
-		for (const FModifyCurveParameters& CurveParameters : SharedData->ModifyCurveParameters)
+		for (const FModifyCurveParameters& CurveParameters : ModifyCurveParameters)
 		{
 			float CurrentValue = Keyframe->Get()->Curves.Get(CurveParameters.CurveName);
-			float UpdatedValue = ProcessCurveOperation(CurrentValue, CurveParameters.CurveValue, InstanceData->Alpha, SharedData->ApplyMode);
+			float UpdatedValue = ProcessCurveOperation(CurrentValue, CurveParameters.CurveValue, Alpha, ApplyMode);
 			Keyframe->Get()->Curves.Set(CurveParameters.CurveName, UpdatedValue);
 		}
 	}
diff --git a/Engine/Plugins/Experimental/UAF/UAFAnimGraph/Source/UAFAnimGraph/Private/Traits/ModifyCurveTrait.h b/Engine/Plugins/Experimental/UAF/UAFAnimGraph/Source/UAFAnimGraph/Private/Traits/ModifyCurveTrait.h
--- a/Engine/Plugins/Experimental/UAF/UAFAnimGraph/Source/UAFAnimGraph/Private/Traits/ModifyCurveTrait.h
+++ b/Engine/Plugins/Experimental/UAF/UAFAnimGraph/Source/UAFAnimGraph/Private/Traits/ModifyCurveTrait.h
@@ -48,7 +48,7 @@ struct FAnimNextModifyCurveTraitSharedData : public FAnimNextTraitSharedData
 	UPROPERTY(Category=Settings, EditAnywhere)
 	EAnimNext_ModifyCurveApplyMode ApplyMode = EAnimNext_ModifyCurveApplyMode::Add;
 
-	UPROPERTY(Category=Settings, EditAnywhere, meta=(Inline, Hidden))
+	UPROPERTY(Category=Settings, EditAnywhere)
 	TArray<FModifyCurveParameters> ModifyCurveParameters;
 
 	UPROPERTY(Category=Settings, EditAnywhere, meta=(ClampMin=0, ClampMax=1))
@@ -57,6 +57,7 @@ struct FAnimNextModifyCurveTraitSharedData : public FAnimNextTraitSharedData
 #define TRAIT_LATENT_PROPERTIES_ENUMERATOR(GeneratorMacro) \
 	GeneratorMacro(ApplyMode) \
 	GeneratorMacro(Alpha) \
+	GeneratorMacro(ModifyCurveParameters) \
 
 	GENERATE_TRAIT_LATENT_PROPERTIES(FAnimNextModifyCurveTraitSharedData, TRAIT_LATENT_PROPERTIES_ENUMERATOR)
 #undef TRAIT_LATENT_PROPERTIES_ENUMERATOR
@@ -78,8 +79,6 @@ namespace UE::UAF
 
 		struct FInstanceData : FTrait::FInstanceData
 		{
-			float Alpha = 1.0f;
-			
 #if ENABLE_ANIM_DEBUG 
 			/** Debug Object for VisualLogger */
 			TObjectPtr<const UObject> HostObject = nullptr;
@@ -98,7 +97,7 @@ struct FModifyCurveTask : public FAnimNextEvaluationTask
 
 	DECLARE_ANIM_EVALUATION_TASK(FModifyCurveTask)
 
-	static FModifyCurveTask Make(UE::UAF::FModifyCurveTrait::FInstanceData* InstanceData, const UE::UAF::FModifyCurveTrait::FSharedData* SharedData);
+	static FModifyCurveTask Make(float Alpha, EAnimNext_ModifyCurveApplyMode ApplyMode, TConstArrayView<FModifyCurveParameters> CurveParameters);
 
 	// Task entry point
 	virtual void Execute(UE::UAF::FEvaluationVM& VM) const override;
@@ -106,6 +105,7 @@ struct FModifyCurveTask : public FAnimNextEvaluationTask
 private:
 	static float ProcessCurveOperation(float CurrentValue, float NewValue, float Alpha, EAnimNext_ModifyCurveApplyMode ApplyMode);
 
-	UE::UAF::FModifyCurveTrait::FInstanceData* InstanceData = nullptr;
-	const UE::UAF::FModifyCurveTrait::FSharedData* SharedData = nullptr;
+	float Alpha = 1.0f;
+	EAnimNext_ModifyCurveApplyMode ApplyMode = EAnimNext_ModifyCurveApplyMode::Add;
+	TArray<FModifyCurveParameters, TInlineAllocator<8>> ModifyCurveParameters;
 };
