UAF: Fix for misaligned trait shared data and instance data, causing crashes when reading non-pinned values.

#Synced-CL 50003294

diff --git a/Engine/Plugins/Experimental/UAF/UAFAnimGraph/Source/UAFAnimGraph/Internal/TraitCore/TraitReader.h b/Engine/Plugins/Experimental/UAF/UAFAnimGraph/Source/UAFAnimGraph/Internal/TraitCore/TraitReader.h
--- a/Engine/Plugins/Experimental/UAF/UAFAnimGraph/Source/UAFAnimGraph/Internal/TraitCore/TraitReader.h
+++ b/Engine/Plugins/Experimental/UAF/UAFAnimGraph/Source/UAFAnimGraph/Internal/TraitCore/TraitReader.h
@@ -35,6 +35,7 @@ namespace UE::UAF
 			GraphTooLarge,				// Exceeded the maximum graph size, @see FTraitReader::MAXIMUM_GRAPH_SHARED_DATA_SIZE
 			NodeSharedDataTooLarge,		// Exceeded the maximum node shared data size, @see FNodeDescription::MAXIMUM_NODE_SHARED_DATA_SIZE
 			NodeInstanceDataTooLarge,	// Exceeded the maximum node instance data size, @see FNodeInstance::MAXIMUM_NODE_INSTANCE_DATA_SIZE
+			InvalidData,				// Unexpected or corrupted data
 		};
 
 		UE_API explicit FTraitReader(const TArray<TObjectPtr<UObject>>& InGraphReferencedObjects, const TArray<FSoftObjectPath>& InGraphReferencedSoftObjects, FArchive& Ar);
diff --git a/Engine/Plugins/Experimental/UAF/UAFAnimGraph/Source/UAFAnimGraph/Private/TraitCore/NodeTemplate.cpp b/Engine/Plugins/Experimental/UAF/UAFAnimGraph/Source/UAFAnimGraph/Private/TraitCore/NodeTemplate.cpp
--- a/Engine/Plugins/Experimental/UAF/UAFAnimGraph/Source/UAFAnimGraph/Private/TraitCore/NodeTemplate.cpp
+++ b/Engine/Plugins/Experimental/UAF/UAFAnimGraph/Source/UAFAnimGraph/Private/TraitCore/NodeTemplate.cpp
@@ -68,6 +68,7 @@ namespace UE::UAF
 		uint32 SharedDataOffset = sizeof(FNodeDescription);
 		uint32 InstanceDataOffset = sizeof(FNodeInstance);
 		uint32 SharedLatentPropertyHandlesOffset = 0;
+		uint32 MinAlignment32 = alignof(FNodeDescription);
 
 		for (uint32 TraitIndex = 0; TraitIndex < NumTraits_; ++TraitIndex)
 		{
@@ -94,6 +95,9 @@ namespace UE::UAF
 				SharedDataOffset = Align(SharedDataOffset, MemoryLayout.SharedDataAlignment);
 				InstanceDataOffset = Align(InstanceDataOffset, MemoryLayout.InstanceDataAlignment);
 
+				MinAlignment32 = FMath::Max(MinAlignment32, MemoryLayout.SharedDataAlignment);
+				MinAlignment32 = FMath::Max(MinAlignment32, MemoryLayout.InstanceDataAlignment);
+
 				// Save our trait offsets
 				TraitSharedDataOffset = SharedDataOffset;
 				TraitInstanceDataOffset = InstanceDataOffset;
@@ -136,6 +140,7 @@ namespace UE::UAF
 			check(TraitSharedDataOffset <= MAX_uint16);
 			check(TraitSharedLatentPropertyHandlesOffset <= MAX_uint16);
 			check(TraitInstanceDataOffset <= MAX_uint16);
+			check(MinAlignment32 <= MAX_uint8);
 
 			// Update our trait offsets
 			TraitTemplate.NodeSharedOffset = static_cast<uint16>(TraitSharedDataOffset);
@@ -143,6 +148,8 @@ namespace UE::UAF
 			TraitTemplate.NodeInstanceOffset = static_cast<uint16>(TraitInstanceDataOffset);
 		}
 
+		MinAlignment = static_cast<uint8>(MinAlignment32);
+
 		// Make sure we respect our alignment constraints
 		SharedDataOffset = Align(SharedDataOffset, alignof(FNodeDescription));
 
diff --git a/Engine/Plugins/Experimental/UAF/UAFAnimGraph/Source/UAFAnimGraph/Private/TraitCore/Trait.cpp b/Engine/Plugins/Experimental/UAF/UAFAnimGraph/Source/UAFAnimGraph/Private/TraitCore/Trait.cpp
--- a/Engine/Plugins/Experimental/UAF/UAFAnimGraph/Source/UAFAnimGraph/Private/TraitCore/Trait.cpp
+++ b/Engine/Plugins/Experimental/UAF/UAFAnimGraph/Source/UAFAnimGraph/Private/TraitCore/Trait.cpp
@@ -9,6 +9,7 @@ namespace UE::UAF
 	void FTrait::SerializeTraitSharedData(FArchive& Ar, FAnimNextTraitSharedData& SharedData) const
 	{
 		UScriptStruct* SharedDataStruct = GetTraitSharedDataStruct();
+		check(IsAligned(&SharedData, SharedDataStruct->GetMinAlignment()));
 		SharedDataStruct->SerializeItem(Ar, &SharedData, nullptr);
 	}
 
diff --git a/Engine/Plugins/Experimental/UAF/UAFAnimGraph/Source/UAFAnimGraph/Private/TraitCore/TraitReader.cpp b/Engine/Plugins/Experimental/UAF/UAFAnimGraph/Source/UAFAnimGraph/Private/TraitCore/TraitReader.cpp
--- a/Engine/Plugins/Experimental/UAF/UAFAnimGraph/Source/UAFAnimGraph/Private/TraitCore/TraitReader.cpp
+++ b/Engine/Plugins/Experimental/UAF/UAFAnimGraph/Source/UAFAnimGraph/Private/TraitCore/TraitReader.cpp
@@ -84,7 +84,23 @@ namespace UE::UAF
 		uint32 SharedDataSize = 0;
 		for (uint32 NodeIndex = 0; NodeIndex < NumNodes; ++NodeIndex)
 		{
-			if (SharedDataSize > MAXIMUM_GRAPH_SHARED_DATA_SIZE)
+			uint32 NodeTemplateIndex = 0;
+			*this << NodeTemplateIndex;
+
+			if (NodeTemplateIndex >= NumNodeTemplates)
+			{
+				// Corrupted / malformed data
+				return EErrorState::InvalidData;
+			}
+
+			const FNodeTemplate* NodeTemplate = NodeTemplateRegistry.Find(NodeTemplateHandles[NodeTemplateIndex]);
+			check(NodeTemplate != nullptr);
+
+			// Align this node
+			check(NodeTemplate->GetMinAlignment() > 0);
+			SharedDataSize = Align(SharedDataSize, NodeTemplate->GetMinAlignment());
+
+			if ((size_t)SharedDataSize > MAXIMUM_GRAPH_SHARED_DATA_SIZE)
 			{
 				// This graph shared data is too large, we won't be able to create handles to this node
 				return EErrorState::GraphTooLarge;
@@ -93,13 +109,15 @@ namespace UE::UAF
 			NodeHandles.Add(FNodeHandle::FromSharedOffset(SharedDataSize));	// This node starts here
 			check(NodeHandles.Last().IsSharedOffset());
 
-			uint32 NodeTemplateIndex = 0;
-			*this << NodeTemplateIndex;
+			uint32 NodeSharedDataSize = NodeTemplate->GetNodeSharedDataSize();
 
-			const FNodeTemplate* NodeTemplate = NodeTemplateRegistry.Find(NodeTemplateHandles[NodeTemplateIndex]);
-			check(NodeTemplate != nullptr);
+			if ((size_t)SharedDataSize + (size_t)NodeSharedDataSize > MAX_uint32)
+			{
+				// Detect overrun of SharedDataSize
+				return EErrorState::GraphTooLarge;
+			}
 
-			SharedDataSize += NodeTemplate->GetNodeSharedDataSize();
+			SharedDataSize += NodeSharedDataSize;
 		}
 
 		// The shared data size might exceed MAXIMUM_GRAPH_SHARED_DATA_SIZE a little bit
@@ -114,6 +132,12 @@ namespace UE::UAF
 		{
 			// Serialize our node shared data
 			const uint32 SharedDataOffset = NodeHandles[NodeIndex].GetSharedOffset();
+			if (SharedDataOffset >= SharedDataSize)
+			{
+				// Corrupted / malformed data
+				return EErrorState::InvalidData;
+			}
+
 			FNodeDescription* NodeDesc = reinterpret_cast<FNodeDescription*>(&GraphSharedData[SharedDataOffset]);
 			NodeDesc->Serialize(*this);
 
diff --git a/Engine/Plugins/Experimental/UAF/UAFAnimGraph/Source/UAFAnimGraph/Public/TraitCore/NodeTemplate.h b/Engine/Plugins/Experimental/UAF/UAFAnimGraph/Source/UAFAnimGraph/Public/TraitCore/NodeTemplate.h
--- a/Engine/Plugins/Experimental/UAF/UAFAnimGraph/Source/UAFAnimGraph/Public/TraitCore/NodeTemplate.h
+++ b/Engine/Plugins/Experimental/UAF/UAFAnimGraph/Source/UAFAnimGraph/Public/TraitCore/NodeTemplate.h
@@ -47,6 +47,9 @@ namespace UE::UAF
 		// Returns the number of traits present in the node template
 		uint32 GetNumTraits() const { return NumTraits; }
 
+		// Returns the required alignment of this node
+		uint32 GetMinAlignment() const { return MinAlignment; }
+
 		// Returns whether the node template is valid or not
 		bool IsValid() const { return NodeSharedDataSize != 0 && NodeInstanceDataSize != 0; }
 
@@ -79,7 +82,8 @@ namespace UE::UAF
 		uint16	NodeInstanceDataSize;	// size in bytes of a node instance, excludes optional latent properties (not serialized, @see FNodeTemplate::Finalize)
 
 		uint8	NumTraits;
-		uint8	Padding[3];
+		uint8	MinAlignment;
+		uint8	Padding[2];
 
 		// TODO: We could use the padding (and extra space) here to cache which trait handles which interface
 		// This would avoid the need to iterate on every trait to look it up. Perhaps only common interfaces could be cached.
diff --git a/Engine/Plugins/Experimental/UAF/UAFAnimGraph/Source/UAFAnimGraph/Public/TraitCore/TraitBinding.h b/Engine/Plugins/Experimental/UAF/UAFAnimGraph/Source/UAFAnimGraph/Public/TraitCore/TraitBinding.h
--- a/Engine/Plugins/Experimental/UAF/UAFAnimGraph/Source/UAFAnimGraph/Public/TraitCore/TraitBinding.h
+++ b/Engine/Plugins/Experimental/UAF/UAFAnimGraph/Source/UAFAnimGraph/Public/TraitCore/TraitBinding.h
@@ -63,7 +63,9 @@ namespace UE::UAF
 			const FTraitTemplate* TraitDescs = Stack->NodeTemplate->GetTraits();
 			const FTraitTemplate* TraitTemplate = TraitDescs + TraitIndex;
 
-			return static_cast<const SharedDataType*>(TraitTemplate->GetTraitDescription(*Stack->NodeDescription));
+			const SharedDataType* SharedData = static_cast<const SharedDataType*>(TraitTemplate->GetTraitDescription(*Stack->NodeDescription));
+			check(IsAligned(SharedData, alignof(SharedDataType)));
+			return SharedData;
 		}
 
 		// Queries a node for a pointer to its trait instance data.
