UAF: Add Apply Additive Per Bone Trait

- Adds new trait that allows the user to apply an additive with a hierarchy table based mask
- Adds blend profile based 'apply additive per bone task' that also handles attribute and curve masks

#Synced-CL 50246811

diff --git a/Engine/Plugins/Experimental/UAF/UAFAnimGraph/Source/UAFAnimGraph/Private/EvaluationVM/Tasks/ApplyAdditiveKeyframePerBone.cpp b/Engine/Plugins/Experimental/UAF/UAFAnimGraph/Source/UAFAnimGraph/Private/EvaluationVM/Tasks/ApplyAdditiveKeyframePerBone.cpp
--- a/Engine/Plugins/Experimental/UAF/UAFAnimGraph/Source/UAFAnimGraph/Private/EvaluationVM/Tasks/ApplyAdditiveKeyframePerBone.cpp
+++ b/Engine/Plugins/Experimental/UAF/UAFAnimGraph/Source/UAFAnimGraph/Private/EvaluationVM/Tasks/ApplyAdditiveKeyframePerBone.cpp
@@ -2,11 +2,17 @@
 
 #include "EvaluationVM/Tasks/ApplyAdditiveKeyframePerBone.h"
 
+#include "HierarchyTable.h"
+#include "HierarchyTableBlendProfile.h"
 #include "EvaluationVM/EvaluationVM.h"
 #include "EvaluationVM/KeyframeState.h"
 #include "TransformArrayOperations.h"
 #include "Animation/AnimationAsset.h"
 #include "UAF/ValueRuntime/ValueBundle.h"
+#include "Animation/AttributeTypes.h"
+#include "Animation/BlendProfile.h"
+#include "Animation/SkeletonRemappingRegistry.h"
+#include "MaskProfile/HierarchyTableTypeMask.h"
 #include "UAF/ValueRuntime/Transformers/AdditiveSpace.h"
 #include "UAF/ValueRuntime/Transformers/BoneSpace.h"
 
@@ -228,3 +234,324 @@ void FUAFApplyAdditiveKeyframePerBoneTask::Execute(UE::UAF::FEvaluationVM& VM) c
 		VM.PushValue(KEYFRAME_STACK_NAME, MoveTemp(BaseKeyframe));
 	}
 }
+
+FUAFApplyAdditiveKeyframeWithBlendProfileTask FUAFApplyAdditiveKeyframeWithBlendProfileTask::Make(const float InBlendWeight, const TSharedPtr<const IBlendProfileInterface>& InBlendProfileInterface)
+{
+	FUAFApplyAdditiveKeyframeWithBlendProfileTask Task;
+	Task.BlendWeight = InBlendWeight;
+	Task.BlendProfileInterface = InBlendProfileInterface;
+	return Task;
+}
+
+void FUAFApplyAdditiveKeyframeWithBlendProfileTask::Execute(UE::UAF::FEvaluationVM& VM) const
+{
+	using namespace UE::UAF;
+
+	if (!BlendProfileInterface.IsValid())
+	{
+		// If we do not have a valid blend profile - fall back to a full body additive
+		UE_LOG(LogAnimation, Warning, TEXT("UAFApplyAdditiveKeyframeWithBlendProfileTask::Execute: No valid blend profile provided. Falling back to full body additive."));
+		Super::Execute(VM);
+		return;
+	}
+	
+	// Currently this task only supports HierarchyTable based blend profiles
+	const TSharedPtr<const FHierarchyTableBlendProfile> HTBlendProfile = StaticCastSharedPtr<const FHierarchyTableBlendProfile>(BlendProfileInterface);
+	if (!HTBlendProfile.IsValid())
+	{
+		// If we do not have a valid blend profile - fall back to a full body additive
+		UE_LOG(LogAnimation, Warning, TEXT("UAFApplyAdditiveKeyframeWithBlendProfileTask::Execute: This task currently only supports HierarchyTable based blend profiles. Falling back to full body additive."));
+		Super::Execute(VM);
+		return;
+	}
+	
+	const USkeleton* SourceSkeleton = BlendProfileInterface->GetSkeleton();
+	const UE::Anim::TNamedValueArray<FDefaultAllocator, UE::Anim::FCurveElement>& CurveMaskWeights = HTBlendProfile->GetCurveBlendWeights(); 
+	const TArray<FHierarchyTableBlendProfile::FMaskedAttributeWeight>& AttributeMaskWeights = HTBlendProfile->GetAttributeBlendWeights();
+	
+	if (VM.GetActiveNamedSet())
+	{
+		// Pop our top two poses, the top one being the additive pose and the second one the base
+		TUniquePtr<FValueBundle> AdditiveInput;
+		if (!VM.PopValue(ATTRIBUTE_STACK_NAME, AdditiveInput))
+		{
+			// We have no inputs, nothing to do
+			return;
+		}
+
+		TUniquePtr<FValueBundle> BaseInput;
+		if (!VM.PopValue(ATTRIBUTE_STACK_NAME, BaseInput))
+		{
+			// We have a single input, discard it since it must be the additive pose, either way something went wrong
+			// Push the reference pose since we'll expect a non-additive pose
+			FValueBundleStack Collection(VM.GetActiveNamedSet());
+			Collection.InitWithValueSpace(FValueSpace(EValueSpaceType::Local));
+
+			VM.PushValue(ATTRIBUTE_STACK_NAME, MakeUnique<FValueBundle>(MoveTemp(Collection)));
+			UE_LOG(LogAnimation, Warning, TEXT("UAFApplyAdditiveKeyframeWithBlendProfileTask::Execute: Could not apply additive keyframe - only a single input was provided. Pushing a ref pose instead."));
+			return;
+		}
+
+		const FValueSpace AdditiveSpace = AdditiveInput->GetValueSpace();
+		if (!AdditiveSpace.IsAdditive())
+		{
+			// Additive input must be additive, push reference pose if not the case
+			FValueBundleStack Collection(VM.GetActiveNamedSet());
+			Collection.InitWithValueSpace(FValueSpace(EValueSpaceType::Local));
+
+			VM.PushValue(ATTRIBUTE_STACK_NAME, MakeUnique<FValueBundle>(MoveTemp(Collection)));
+			UE_LOG(LogAnimation, Warning, TEXT("UAFApplyAdditiveKeyframeWithBlendProfileTask::Execute: Could not apply additive keyframe - The expected additive pose was not of an additive type. Pushing a ref pose instead."));
+			return;
+		}
+		
+		const float AdditiveWeight = GetInterpolationAlpha(*BaseInput, *AdditiveInput);
+	
+		UScriptStruct* FloatAttributeType = FFloatAnimationAttribute::StaticStruct();
+
+		// Build our per-value weights (only for bound values)
+		FValueBundleStack PerValueWeights(VM.GetActiveNamedSet());
+		FBoundMapCollection& PerValueWeightMaps = PerValueWeights.GetBoundValueMaps();
+
+		for (auto It = BaseInput->GetBoundValueMaps().CreateConstIterator(); It; ++It)
+		{
+			UScriptStruct* AttributeType = It.GetAttributeType();
+
+			// Every attribute type maps to a float, our weight
+			FAttributeMappingKey MappingKey = FAttributeMappingKey::MakeFromTo(AttributeType, FloatAttributeType);
+
+			TBoundValueMap<FFloatAnimationAttribute>* SetWeights = PerValueWeightMaps.Add<FFloatAnimationAttribute>(MappingKey);
+			
+			// Bone transforms
+			if (const TBoundValueMap<FBoneTransformAnimationAttribute>* BoneMap = Cast<FBoneTransformAnimationAttribute>(It.GetMap()))
+			{
+				const FAttributeTypedSetPtr& BoneSet = BoneMap->GetTypedSet();
+				const int32 NumBonesInSet = BoneMap->Num();
+
+				check(BlendProfileInterface->GetNumBlendEntries() >= NumBonesInSet);
+				
+				const USkeleton* TargetSkeleton = VM.GetActiveEvaluationContext().GetSkeleton();
+				const FSkeletonRemapping& SkeletonRemapping = UE::Anim::FSkeletonRemappingRegistry::Get().GetRemapping(SourceSkeleton, TargetSkeleton);
+				
+				for (FAttributeSetIndex SetIndex(0); SetIndex < NumBonesInSet; ++SetIndex)
+				{
+					const FAttributeBindingIndex BindingIndex = BoneSet->GetBindingIndex(SetIndex);
+					const FSkeletonPoseBoneIndex TargetSkeletonBoneIndex = BindingIndex;
+					int32 WeightIndex = TargetSkeletonBoneIndex.GetInt();
+					
+					if (SkeletonRemapping.IsValid())
+					{
+						const FSkeletonPoseBoneIndex SourceSkeletonBoneIndex(SkeletonRemapping.GetSourceSkeletonBoneIndex(TargetSkeletonBoneIndex.GetInt()));
+						WeightIndex = SourceSkeletonBoneIndex.IsValid() ? SourceSkeletonBoneIndex.GetInt() : INDEX_NONE;
+					}
+					;
+					
+					(*SetWeights)[SetIndex].Value = WeightIndex != INDEX_NONE ? BlendProfileInterface->GetBoneBlendScale(WeightIndex) : 0.0f;
+					(*SetWeights)[SetIndex].Value *= AdditiveWeight;
+				}
+			}
+			// Curves and attributes
+			else if (const TBoundValueMap<FFloatAnimationAttribute>* FloatMap = Cast<FFloatAnimationAttribute>(It.GetMap()))
+			{
+				SetWeights->FillWithValue(FFloatAnimationAttribute{ 0.0f });
+				
+				// Curves
+				const FAttributeTypedSetPtr& FloatSet = FloatMap->GetTypedSet();
+				CurveMaskWeights.ForEachElement([&FloatSet, SetWeights, AdditiveWeight](const UE::Anim::FCurveElement& Curve)
+					{
+						const FAttributeSetIndex SetIndex = FloatSet->FindIndex(Curve.Name);
+						if (SetIndex.IsValid())
+						{
+							(*SetWeights)[SetIndex].Value = Curve.Value;
+							(*SetWeights)[SetIndex].Value *= AdditiveWeight;
+						}
+					});
+				
+				// Attributes
+				for (const FHierarchyTableBlendProfile::FMaskedAttributeWeight& AttributeWeight : AttributeMaskWeights)
+				{
+					const FAttributeSetIndex SetIndex = FloatSet->FindIndex(AttributeWeight.Attribute.GetName());
+					if (SetIndex.IsValid())
+					{
+						(*SetWeights)[SetIndex].Value = AttributeWeight.Weight;
+						(*SetWeights)[SetIndex].Value *= AdditiveWeight;
+					}
+				}
+			}
+			else
+			{
+				SetWeights->FillWithValue(FFloatAnimationAttribute{ 0.0f });
+
+				const FAttributeTypedSetPtr& TypedSet = It.GetMap()->GetTypedSet();
+				for (const FHierarchyTableBlendProfile::FMaskedAttributeWeight& AttributeWeight : AttributeMaskWeights)
+				{
+					const FAttributeSetIndex SetIndex = TypedSet->FindIndex(AttributeWeight.Attribute.GetName());
+					if (SetIndex.IsValid())
+					{
+						(*SetWeights)[SetIndex].Value = AttributeWeight.Weight;
+					}
+				}
+			}
+		}
+
+		bool bConvertBaseToMeshSpace = false;
+		if (AdditiveSpace.GetMixedSpaceFlags() == EMixedSpaceFlags::MeshRotation)
+		{
+			// If our additive input is in mesh rotation space, we ensure our base is in that space as well
+			const FValueSpace BaseSpace = BaseInput->GetValueSpace();
+			bConvertBaseToMeshSpace = BaseSpace.GetMixedSpaceFlags() != EMixedSpaceFlags::MeshRotation;
+		}
+
+		if (bConvertBaseToMeshSpace)
+		{
+			Transformers::FBoneSpace::LocalToMeshRotation(*BaseInput, *BaseInput);
+		}
+
+		Transformers::FApplyAdditiveSpace::Apply(VM.GetTransformerMap(), *BaseInput, *AdditiveInput, PerValueWeights, AdditiveWeight, *BaseInput);
+
+		if (bConvertBaseToMeshSpace)
+		{
+			Transformers::FBoneSpace::MeshRotationToLocal(*BaseInput, *BaseInput);
+		}
+
+		VM.PushValue(ATTRIBUTE_STACK_NAME, MoveTemp(BaseInput));
+	}
+	else
+	{
+		// Pop our top two poses, the top one being the additive pose and the second one the base
+		TUniquePtr<FKeyframeState> AdditiveKeyframe;
+		if (!VM.PopValue(KEYFRAME_STACK_NAME, AdditiveKeyframe))
+		{
+			// We have no inputs, nothing to do
+			return;
+		}
+
+		TUniquePtr<FKeyframeState> BaseKeyframe;
+		if (!VM.PopValue(KEYFRAME_STACK_NAME, BaseKeyframe))
+		{
+			// We have a single input, discard it since it must be the additive pose, either way something went wrong
+			// Push the reference pose since we'll expect a non-additive pose
+			VM.PushValue(KEYFRAME_STACK_NAME, MakeUnique<FKeyframeState>(VM.MakeReferenceKeyframe(false)));
+			UE_LOG(LogAnimation, Warning, TEXT("UAFApplyAdditiveKeyframeWithBlendProfileTask::Execute: Could not apply additive keyframe - only a single input was provided. Pushing a ref pose instead."));
+			return;
+		}
+
+		if (!AdditiveKeyframe->Pose.IsAdditive())
+		{
+			// Additive must be additive type, push reference pose if not the case
+			VM.PushValue(KEYFRAME_STACK_NAME, MakeUnique<FKeyframeState>(VM.MakeReferenceKeyframe(false)));
+			UE_LOG(LogAnimation, Warning, TEXT("UAFApplyAdditiveKeyframeWithBlendProfileTask::Execute: Could not apply additive keyframe - The expected additive pose was not of an additive type. Pushing a ref pose instead."));
+			return;
+		}
+
+		const float AdditiveWeight = GetInterpolationAlpha(BaseKeyframe.Get(), AdditiveKeyframe.Get());
+		
+		const TArrayView<const FBoneIndexType> LODBoneIndexToSkeletonBoneIndexMap = BaseKeyframe->Pose.GetLODBoneIndexToSkeletonBoneIndexMap();
+		const int32 NumLODBones = LODBoneIndexToSkeletonBoneIndexMap.Num();
+		
+		TArray<int32> LODBoneIndexToWeightIndexMap;
+		LODBoneIndexToWeightIndexMap.AddUninitialized(NumLODBones);
+		
+		const USkeleton* TargetSkeleton = BaseKeyframe->Pose.GetSkeletonAsset();
+		TSharedPtr<IInterpolationIndexProvider::FPerBoneInterpolationData> Data = HTBlendProfile->GetPerBoneInterpolationData(TargetSkeleton);
+		const FSkeletonRemapping& SkeletonRemapping = UE::Anim::FSkeletonRemappingRegistry::Get().GetRemapping(SourceSkeleton, TargetSkeleton);
+		if (SkeletonRemapping.IsValid())
+		{
+			for (int32 LODBoneIndex = 0; LODBoneIndex < NumLODBones; ++LODBoneIndex)
+			{
+				const FSkeletonPoseBoneIndex TargetSkeletonBoneIndex(LODBoneIndexToSkeletonBoneIndexMap[LODBoneIndex]);
+				const FSkeletonPoseBoneIndex SourceSkeletonBoneIndex(SkeletonRemapping.GetSourceSkeletonBoneIndex(TargetSkeletonBoneIndex.GetInt()));
+				
+				LODBoneIndexToWeightIndexMap[LODBoneIndex] = SourceSkeletonBoneIndex.IsValid() ? HTBlendProfile->GetPerBoneInterpolationIndex(SourceSkeletonBoneIndex, TargetSkeleton, Data.Get()) : INDEX_NONE;
+			}
+		}
+		else
+		{
+			for (int32 LODBoneIndex = 0; LODBoneIndex < NumLODBones; ++LODBoneIndex)
+			{
+				const FSkeletonPoseBoneIndex SkeletonBoneIndex(LODBoneIndexToSkeletonBoneIndexMap[LODBoneIndex]);
+				LODBoneIndexToWeightIndexMap[LODBoneIndex] = SkeletonBoneIndex.IsValid() ? HTBlendProfile->GetPerBoneInterpolationIndex(SkeletonBoneIndex, TargetSkeleton, Data.Get()) : INDEX_NONE;
+			}
+		}
+		
+		if (EnumHasAnyFlags(VM.GetFlags(), EEvaluationFlags::Bones))
+		{
+			check(BaseKeyframe->Pose.GetNumBones() == AdditiveKeyframe->Pose.GetNumBones());
+			
+			TArray<float> BoneWeights;
+			BoneWeights.SetNum(BlendProfileInterface->GetNumBlendEntries());
+			for (int EntryIndex = 0; EntryIndex < BlendProfileInterface->GetNumBlendEntries(); EntryIndex++)
+			{
+				BoneWeights[EntryIndex] = BlendProfileInterface->GetBoneBlendScale(EntryIndex) * AdditiveWeight;
+			}
+			
+			const FTransformArrayView BaseTransformsView = BaseKeyframe->Pose.LocalTransforms.GetView();
+
+			if (AdditiveKeyframe->Pose.IsMeshSpaceAdditive())
+			{
+				BlendWithIdentityAndAccumulatePerBoneWeightMesh(
+					BaseTransformsView, AdditiveKeyframe->Pose.LocalTransforms.GetConstView(),
+					AdditiveKeyframe->Pose.GetLODBoneIndexToParentLODBoneIndexMap(), LODBoneIndexToWeightIndexMap, BoneWeights, AdditiveWeight);
+			}
+			else
+			{
+				BlendWithIdentityAndAccumulatePerBoneWeight(BaseTransformsView, AdditiveKeyframe->Pose.LocalTransforms.GetConstView(), LODBoneIndexToWeightIndexMap, BoneWeights, AdditiveWeight);
+			}
+
+			NormalizeRotations(BaseTransformsView);
+		}
+
+		if (EnumHasAnyFlags(VM.GetFlags(), EEvaluationFlags::Curves))
+		{
+			if (FAnimWeight::IsRelevant(AdditiveWeight))
+			{
+				FBlendedCurve& BaseCurve = BaseKeyframe->Curves;
+				FBlendedCurve& AdditiveCurve = AdditiveKeyframe->Curves;
+				
+				// Apply the mask weight to matching additive curves
+				FBlendedCurve MaskedAdditiveCurve;
+				UE::Anim::FNamedValueArrayUtils::Intersection(AdditiveCurve, CurveMaskWeights, [&MaskedAdditiveCurve] (const UE::Anim::FCurveElement& InAdditiveElement, const UE::Anim::FCurveElement& InMaskElement) mutable
+				{
+					const float MaskedValue = InAdditiveElement.Value * (InMaskElement.Value);
+					MaskedAdditiveCurve.Add(InAdditiveElement.Name, MaskedValue);
+				});
+				
+				AdditiveCurve.Combine(MaskedAdditiveCurve);
+				
+				// Accumulate curves values by the additive weight 
+				UE::Anim::FNamedValueArrayUtils::Union(
+					BaseCurve, 
+					AdditiveCurve, 
+					[AdditiveWeight](UE::Anim::FCurveElement& InOutThisElement, const UE::Anim::FCurveElement& InAdditiveCurveElement, UE::Anim::ENamedValueUnionFlags InFlags)
+				{
+					if (InFlags == UE::Anim::ENamedValueUnionFlags::ValidArg0 || InFlags == UE::Anim::ENamedValueUnionFlags::ValidArg1)
+					{
+						InOutThisElement.Value += InAdditiveCurveElement.Value * AdditiveWeight;
+						InOutThisElement.Flags |= InAdditiveCurveElement.Flags;
+						
+					}
+				});
+			}
+		}
+
+		if (EnumHasAnyFlags(VM.GetFlags(), EEvaluationFlags::Attributes))
+		{
+			// we could accumulate all attributes per additive weight and then just correct the masked ones manually 
+			UE::Anim::FStackAttributeContainer& BaseAttributes = BaseKeyframe->Attributes;
+			UE::Anim::FStackAttributeContainer& AdditiveAttributes = AdditiveKeyframe->Attributes;
+			
+			// Build up override weight map for masked attributes 
+			TMap<UE::Anim::FAttributeId, float> AttributeMaskTotalWeights;
+			for (const FHierarchyTableBlendProfile::FMaskedAttributeWeight& AttributeWeight : AttributeMaskWeights)
+			{
+				AttributeMaskTotalWeights.Emplace(AttributeWeight.Attribute, AttributeWeight.Weight * AdditiveWeight);
+			}
+			
+			// Accumulates all attributes with additive weight, writes into BaseAttributes
+			const EAdditiveAnimationType AdditiveAnimationType = AdditiveKeyframe->Pose.IsMeshSpaceAdditive() ? AAT_RotationOffsetMeshSpace : AAT_LocalSpaceBase;
+			UE::Anim::Attributes::AccumulateAttributesWithOverrideWeights(AdditiveAttributes, BaseAttributes, AttributeMaskTotalWeights, AdditiveWeight, AdditiveAnimationType);
+		}
+
+		VM.PushValue(KEYFRAME_STACK_NAME, MoveTemp(BaseKeyframe));
+	}
+}
+
diff --git a/Engine/Plugins/Experimental/UAF/UAFAnimGraph/Source/UAFAnimGraph/Private/Traits/ApplyAdditivePerBoneTrait.cpp b/Engine/Plugins/Experimental/UAF/UAFAnimGraph/Source/UAFAnimGraph/Private/Traits/ApplyAdditivePerBoneTrait.cpp
new file mode 100644
--- /dev/null
+++ b/Engine/Plugins/Experimental/UAF/UAFAnimGraph/Source/UAFAnimGraph/Private/Traits/ApplyAdditivePerBoneTrait.cpp
@@ -0,0 +1,166 @@
+// Copyright Epic Games, Inc. All Rights Reserved.
+
+
+#include "Traits/ApplyAdditivePerBoneTrait.h"
+
+#include "HierarchyTable.h"
+#include "EvaluationVM/Tasks/ApplyAdditiveKeyframePerBone.h"
+#include "EvaluationVM/Tasks/ConvertRotationsToLocalSpace.h"
+#include "EvaluationVM/Tasks/ConvertRotationsToMeshSpace.h"
+#include "EvaluationVM/Tasks/NormalizeRotations.h"
+
+namespace UE::UAF
+{
+	
+AUTO_REGISTER_ANIM_TRAIT(FApplyAdditivePerBoneTrait)
+
+#define TRAIT_INTERFACE_ENUMERATOR(GeneratorMacro) \
+GeneratorMacro(IContinuousBlend) \
+GeneratorMacro(IEvaluate) \
+GeneratorMacro(IHierarchy) \
+GeneratorMacro(IUpdate) \
+GeneratorMacro(IUpdateTraversal) \
+
+GENERATE_ANIM_TRAIT_IMPLEMENTATION(FApplyAdditivePerBoneTrait, TRAIT_INTERFACE_ENUMERATOR, NULL_ANIM_TRAIT_INTERFACE_ENUMERATOR, NULL_ANIM_TRAIT_EVENT_ENUMERATOR)
+#undef TRAIT_INTERFACE_ENUMERATOR
+	
+void FApplyAdditivePerBoneTrait::PreUpdate(FUpdateTraversalContext& Context, const TTraitBinding<IUpdate>& Binding,	const FTraitUpdateState& TraitState) const
+{
+	const FSharedData* SharedData = Binding.GetSharedData<FSharedData>();
+	FInstanceData* InstanceData = Binding.GetInstanceData<FInstanceData>();
+	
+	TTraitBinding<IContinuousBlend> ContinuousBlendTrait;
+	Binding.GetStackInterface(ContinuousBlendTrait);
+	
+	const TObjectPtr<UHierarchyTable> BlendProfile = SharedData->GetBlendProfile(Binding);
+	if (!InstanceData->Base.IsValid())
+	{
+		InstanceData->Base = Context.AllocateNodeInstance(Binding, SharedData->Base);
+	}
+	else
+	{
+		InstanceData->bWasBaseRelevant = true;
+	}
+	
+	const float AdditiveWeight = ContinuousBlendTrait.GetBlendWeight(Context, 1);
+	if (FAnimWeight::IsRelevant(AdditiveWeight) && BlendProfile)
+	{
+		if (!InstanceData->Additive.IsValid())
+		{
+			// We need to blend a child that isn't instanced yet, allocate it
+			InstanceData->Additive = Context.AllocateNodeInstance(Binding, SharedData->Additive);
+		}
+		else
+		{
+			InstanceData->bWasAdditiveRelevant = true;
+		}
+		
+		// Check if the blend profile has changed since the last update and update relevant instance data accordingly
+		if (!InstanceData->CachedBlendProfile.IsValid() || InstanceData->CachedBlendProfile.Get() != BlendProfile)
+		{
+			InstanceData->CachedBlendProfile = BlendProfile;
+			InstanceData->HierarchyTableBlendProfile = MakeShared<FHierarchyTableBlendProfile>(BlendProfile, EBlendProfileMode::BlendMask);
+		}
+	}
+	else
+	{
+		InstanceData->bWasAdditiveRelevant = false;
+		InstanceData->Additive.Reset();
+	}
+}
+
+
+void FApplyAdditivePerBoneTrait::PostEvaluate(FEvaluateTraversalContext& Context, const TTraitBinding<IEvaluate>& Binding) const
+{
+	const FSharedData* SharedData = Binding.GetSharedData<FSharedData>();
+	const FInstanceData* InstanceData = Binding.GetInstanceData<FInstanceData>();
+	
+	if (InstanceData->Base.IsValid() && InstanceData->Additive.IsValid() && InstanceData->HierarchyTableBlendProfile.IsValid())
+	{
+		TTraitBinding<IContinuousBlend> ContinuousBlendTrait;
+		Binding.GetStackInterface(ContinuousBlendTrait);
+		
+		const float AdditiveWeight = ContinuousBlendTrait.GetBlendWeight(Context, 1);
+		const bool bBlendInMeshSpace = SharedData->GetbMeshSpaceBlend(Binding);
+
+		if (bBlendInMeshSpace)
+		{
+			Context.AppendTask(FAnimNextConvertRotationsToMeshSpaceTask::Make(2));
+		}
+
+		Context.AppendTask(FUAFApplyAdditiveKeyframeWithBlendProfileTask::Make(AdditiveWeight, InstanceData->HierarchyTableBlendProfile));
+		Context.AppendTask(FAnimNextNormalizeKeyframeRotationsTask());
+
+		if (bBlendInMeshSpace)
+		{
+			Context.AppendTask(FAnimNextConvertRotationsToLocalSpaceTask::Make(1));
+		}
+	}
+	else
+	{
+		// We have only one child that is active, do nothing
+	}
+}
+	
+void FApplyAdditivePerBoneTrait::QueueChildrenForTraversal(FUpdateTraversalContext& Context, const TTraitBinding<IUpdateTraversal>& Binding,
+	const FTraitUpdateState& TraitState, FUpdateTraversalQueue& TraversalQueue) const
+{
+	const FSharedData* SharedData = Binding.GetSharedData<FSharedData>();
+	const FInstanceData* InstanceData = Binding.GetInstanceData<FInstanceData>();
+	
+	const TObjectPtr<UHierarchyTable> BlendProfile = SharedData->GetBlendProfile(Binding);
+	
+	TTraitBinding<IContinuousBlend> ContinuousBlendTrait;
+	Binding.GetStackInterface(ContinuousBlendTrait);
+
+	const float BaseWeight = ContinuousBlendTrait.GetBlendWeight(Context, 0);
+	const float AdditiveWeight = ContinuousBlendTrait.GetBlendWeight(Context, 1);
+	if (InstanceData->Base.IsValid())
+	{
+		TraversalQueue.Push(InstanceData->Base, TraitState.WithWeight(BaseWeight).AsNewlyRelevant(!InstanceData->bWasBaseRelevant));
+	}
+
+	if (InstanceData->Additive.IsValid() && FAnimWeight::IsRelevant(AdditiveWeight) && BlendProfile)
+	{
+		TraversalQueue.Push(InstanceData->Additive, TraitState.WithWeight(AdditiveWeight).AsNewlyRelevant(!InstanceData->bWasAdditiveRelevant));
+	}
+}
+
+
+uint32 FApplyAdditivePerBoneTrait::GetNumChildren(const FExecutionContext& Context, const TTraitBinding<IHierarchy>& Binding) const
+{
+	return 2;
+}
+
+void FApplyAdditivePerBoneTrait::GetChildren(const FExecutionContext& Context, const TTraitBinding<IHierarchy>& Binding, FChildrenArray& Children) const
+{
+	const FInstanceData* InstanceData = Binding.GetInstanceData<FInstanceData>();
+
+	// Add the two children, even if the handles are empty
+	Children.Add(InstanceData->Base);
+	Children.Add(InstanceData->Additive);
+}
+
+float FApplyAdditivePerBoneTrait::GetBlendWeight(const FExecutionContext& Context, const TTraitBinding<IContinuousBlend>& Binding, int32 ChildIndex) const
+{
+	const FSharedData* SharedData = Binding.GetSharedData<FSharedData>();
+	float Alpha = FMath::Clamp(SharedData->GetAlpha(Binding), 0.0f, 1.0f);
+	if (SharedData->GetBlendProfile(Binding) == nullptr)
+	{
+		Alpha = 0.0f;
+	}
+
+	if (ChildIndex == 0)
+	{
+		return 1.0f - Alpha;
+	}
+	
+	if (ChildIndex == 1)
+	{
+		return Alpha;
+	}
+
+	// Invalid child index
+	return -1.0f;
+}
+}
diff --git a/Engine/Plugins/Experimental/UAF/UAFAnimGraph/Source/UAFAnimGraph/Private/Traits/ApplyAdditivePerBoneTrait.h b/Engine/Plugins/Experimental/UAF/UAFAnimGraph/Source/UAFAnimGraph/Private/Traits/ApplyAdditivePerBoneTrait.h
new file mode 100644
--- /dev/null
+++ b/Engine/Plugins/Experimental/UAF/UAFAnimGraph/Source/UAFAnimGraph/Private/Traits/ApplyAdditivePerBoneTrait.h
@@ -0,0 +1,55 @@
+// Copyright Epic Games, Inc. All Rights Reserved.
+
+#pragma once
+#include "HierarchyTableBlendProfile.h"
+#include "TraitCore/Trait.h"
+#include "TraitInterfaces/IContinuousBlend.h"
+#include "TraitInterfaces/IEvaluate.h"
+#include "TraitInterfaces/IHierarchy.h"
+#include "TraitInterfaces/IUpdate.h"
+#include "Traits/ApplyAdditivePerBoneTraitData.h"
+
+
+namespace UE::UAF
+{
+	/**
+	 * FApplyAdditivePerBone Trait
+	 * 
+	 * A trait that can apply an additive animation with a hierarchy table based mask
+	 */
+	struct FApplyAdditivePerBoneTrait : FBaseTrait, IEvaluate, IUpdate, IUpdateTraversal, IHierarchy, IContinuousBlend
+	{
+		DECLARE_ANIM_TRAIT(FApplyAdditivePerBoneTrait, FBaseTrait)
+		
+		using FSharedData = FApplyAdditivePerBoneTraitSharedData;
+
+		struct FInstanceData : FTrait::FInstanceData
+		{
+			FTraitPtr Base;
+			FTraitPtr Additive;
+
+			bool bWasBaseRelevant = false;
+			bool bWasAdditiveRelevant = false;
+			
+			TSharedPtr<const FHierarchyTableBlendProfile> HierarchyTableBlendProfile = nullptr;
+			TWeakObjectPtr<UHierarchyTable> CachedBlendProfile = nullptr;
+		};
+
+
+		// IEvaluate impl
+		virtual void PostEvaluate(FEvaluateTraversalContext& Context, const TTraitBinding<IEvaluate>& Binding) const override;
+
+		// IUpdate impl
+		virtual void PreUpdate(FUpdateTraversalContext& Context, const TTraitBinding<IUpdate>& Binding, const FTraitUpdateState& TraitState) const override;
+
+		// IUpdateTraversal impl
+		virtual void QueueChildrenForTraversal(FUpdateTraversalContext& Context, const TTraitBinding<IUpdateTraversal>& Binding, const FTraitUpdateState& TraitState, FUpdateTraversalQueue& TraversalQueue) const override;
+
+		// IHierarchy impl
+		virtual uint32 GetNumChildren(const FExecutionContext& Context, const TTraitBinding<IHierarchy>& Binding) const override;
+		virtual void GetChildren(const FExecutionContext& Context, const TTraitBinding<IHierarchy>& Binding, FChildrenArray& Children) const override;
+
+		// IContinuousBlend impl
+		virtual float GetBlendWeight(const FExecutionContext& Context, const TTraitBinding<IContinuousBlend>& Binding, int32 ChildIndex) const override;
+	};
+}
diff --git a/Engine/Plugins/Experimental/UAF/UAFAnimGraph/Source/UAFAnimGraph/Public/EvaluationVM/Tasks/ApplyAdditiveKeyframePerBone.h b/Engine/Plugins/Experimental/UAF/UAFAnimGraph/Source/UAFAnimGraph/Public/EvaluationVM/Tasks/ApplyAdditiveKeyframePerBone.h
--- a/Engine/Plugins/Experimental/UAF/UAFAnimGraph/Source/UAFAnimGraph/Public/EvaluationVM/Tasks/ApplyAdditiveKeyframePerBone.h
+++ b/Engine/Plugins/Experimental/UAF/UAFAnimGraph/Source/UAFAnimGraph/Public/EvaluationVM/Tasks/ApplyAdditiveKeyframePerBone.h
@@ -10,6 +10,7 @@
 #define UE_API UAFANIMGRAPH_API
 
 struct FBlendSampleData;
+class IBlendProfileInterface;
 
 /*
  * Apply Additive Keyframe Per Bone Task
@@ -36,4 +37,30 @@ struct FUAFApplyAdditiveKeyframePerBoneTask : public FAnimNextApplyAdditiveKeyfr
 	bool bPerBoneDataUsesSkeletonIndex = false;
 };
 
+/*
+ * Apply Additive Keyframe Per Bone Task
+ *
+ * This pops the top two keyframes from the VM keyframe stack, 
+ * it applies an additive keyframe with a weight per bone based on the given blend profile onto its base and pushes back the result onto the stack.
+ * If no valid blend profile is provided it will revert back to a full pose additive application. 
+ * 
+ * The top pose should be the additive keyframe and the second to the top the base keyframe.
+ * 
+ * This task currently only supports hierarchy table based blend profiles.
+ */
+USTRUCT()
+struct FUAFApplyAdditiveKeyframeWithBlendProfileTask : public FAnimNextApplyAdditiveKeyframeTask
+{
+	GENERATED_BODY()
+
+	DECLARE_ANIM_EVALUATION_TASK(FUAFApplyAdditiveKeyframeWithBlendProfileTask)
+	
+	static UE_API FUAFApplyAdditiveKeyframeWithBlendProfileTask Make(const float InBlendWeight, const TSharedPtr<const IBlendProfileInterface>& InBlendProfileInterface);
+	
+	UE_API virtual void Execute(UE::UAF::FEvaluationVM& VM) const override;
+	
+	TSharedPtr<const IBlendProfileInterface> BlendProfileInterface = nullptr;
+};
+
+
 #undef UE_API
diff --git a/Engine/Plugins/Experimental/UAF/UAFAnimGraph/Source/UAFAnimGraph/Public/Traits/ApplyAdditivePerBoneTraitData.h b/Engine/Plugins/Experimental/UAF/UAFAnimGraph/Source/UAFAnimGraph/Public/Traits/ApplyAdditivePerBoneTraitData.h
new file mode 100644
--- /dev/null
+++ b/Engine/Plugins/Experimental/UAF/UAFAnimGraph/Source/UAFAnimGraph/Public/Traits/ApplyAdditivePerBoneTraitData.h
@@ -0,0 +1,51 @@
+// Copyright Epic Games, Inc. All Rights Reserved.
+
+#pragma once
+
+#include "TraitCore/Trait.h"
+#include "TraitCore/TraitBinding.h"
+
+#include "ApplyAdditivePerBoneTraitData.generated.h"
+
+class UHierarchyTable;
+
+namespace UE::UAF
+{
+
+// A trait that can apply a mesh or local space additive with an additional mask.
+USTRUCT(meta = (DisplayName = "Apply Additive Per Bone", ShowTooltip=true))
+struct FApplyAdditivePerBoneTraitSharedData : public FAnimNextTraitSharedData
+{
+	GENERATED_BODY()
+
+	// Base to apply additive to.
+	UPROPERTY()
+	FAnimNextTraitHandle Base;
+
+	// Additive to apply.
+	UPROPERTY()
+	FAnimNextTraitHandle Additive;
+	
+	// TODO: To be changed over to UStandaloneBlendProfile
+	// HierarchyTable Mask asset
+	UPROPERTY(EditAnywhere, Category = "Default")
+	TObjectPtr<UHierarchyTable> BlendProfile = nullptr;
+	
+	// How much to apply of the additive 
+	UPROPERTY(EditAnywhere, Category = "Default")
+	float Alpha = 1.0f;
+	
+	// Whether to apply in mesh or local space
+	UPROPERTY(EditAnywhere, Category = "Default")
+	bool bMeshSpaceBlend = false;
+
+	// Latent pin support boilerplate
+	#define TRAIT_LATENT_PROPERTIES_ENUMERATOR(GeneratorMacro) \
+	GeneratorMacro(Alpha) \
+	GeneratorMacro(BlendProfile) \
+	GeneratorMacro(bMeshSpaceBlend) \
+
+	GENERATE_TRAIT_LATENT_PROPERTIES(FApplyAdditivePerBoneTraitSharedData, TRAIT_LATENT_PROPERTIES_ENUMERATOR)
+	#undef TRAIT_LATENT_PROPERTIES_ENUMERATOR
+};
+}
diff --git a/Engine/Source/Runtime/Engine/Classes/Animation/AttributeBlendData.h b/Engine/Source/Runtime/Engine/Classes/Animation/AttributeBlendData.h
--- a/Engine/Source/Runtime/Engine/Classes/Animation/AttributeBlendData.h
+++ b/Engine/Source/Runtime/Engine/Classes/Animation/AttributeBlendData.h
@@ -40,6 +40,12 @@ namespace UE
 				ensure(SourceAttributes.Num() == 1);
 				return FAttributeBlendData(SourceAttributes, InUniformWeight, AdditiveType, AttributeScriptStruct);
 			}
+			
+			static FAttributeBlendData SingleAdditiveContainerWithOverrideWeights(const TArrayView<const FStackAttributeContainer* const> SourceAttributes, const TMap<UE::Anim::FAttributeId, float>& InAttributeWeightOverrides, const float InUniformWeight, EAdditiveAnimationType AdditiveType, const UScriptStruct* AttributeScriptStruct)
+			{
+				ensure(SourceAttributes.Num() == 1);
+				return FAttributeBlendData(SourceAttributes, InAttributeWeightOverrides, InUniformWeight, AdditiveType, AttributeScriptStruct);
+			}
 
 			static FAttributeBlendData PerBoneWeighted(const TArrayView<const FStackAttributeContainer> SourceAttributes, const TArrayView<const FPerBoneBlendWeight> InPerBoneBlendWeights, const UScriptStruct* AttributeScriptStruct)
 			{
@@ -80,6 +86,9 @@ namespace UE
 
 			// Additive accumulate using a single weight
 			ENGINE_API FAttributeBlendData(const TArrayView<const FStackAttributeContainer* const> SourceAttributes, const float InUniformWeight, EAdditiveAnimationType InAdditiveType, const UScriptStruct* AttributeScriptStruct);
+			
+			// Additive accumulate using a single weight plus optional per attribute override weights 
+			ENGINE_API FAttributeBlendData(TArrayView<const FStackAttributeContainer* const> SourceAttributes, const TMap<UE::Anim::FAttributeId, float>& InAttributeWeightOverrides, const float InUniformWeight, EAdditiveAnimationType InAdditiveType, const UScriptStruct* AttributeScriptStruct);
 
 			// Blend using per-bone blend weights
 			ENGINE_API FAttributeBlendData(const TArrayView<const FStackAttributeContainer> SourceAttributes, const TArrayView<const FPerBoneBlendWeight> InPerBoneBlendWeights, const UScriptStruct* AttributeScriptStruct);
@@ -163,6 +172,10 @@ namespace UE
 			TArrayView<const int32> PerBoneInterpolationIndices;
 			TArrayView<const FBlendSampleData> BlendSampleDataCache;
 			inline const TArrayView<const int32>& GetBlendSampleDataCacheIndices() const { return WeightIndices; }
+			
+			/* Contains optional per attribute weight overrides */
+			TMap<UE::Anim::FAttributeId, float> AttributeWeightOverrides;
+			
 		public:
 			template<typename AttributeType>
 			struct TAttributeSetIterator
@@ -456,6 +469,17 @@ namespace UE
 				{
 					return Data.bPerBoneFilter;
 				}
+				
+				const float GetWeightOverride(const FAttributeId& InAttribute) const
+				{
+					const float* FoundWeight = Data.AttributeWeightOverrides.Find(InAttribute);
+					if (FoundWeight)
+					{
+						return *FoundWeight; 
+					}
+					
+					return -1.0f;
+				}
 
 			protected:
 				/** Outer object that creates this */
@@ -528,6 +552,17 @@ namespace UE
 				{
 					return Data.bPerBoneFilter;
 				}
+				
+				const float GetWeightOverride(const FAttributeId& InAttribute) const
+				{
+					const float* FoundWeight = Data.AttributeWeightOverrides.Find(InAttribute);
+					if (FoundWeight)
+					{
+						return *FoundWeight; 
+					}
+					
+					return -1.0f;
+				}
 
 			protected:
 				/** Outer object that creates this */
diff --git a/Engine/Source/Runtime/Engine/Classes/Animation/AttributeBlendOperator.h b/Engine/Source/Runtime/Engine/Classes/Animation/AttributeBlendOperator.h
--- a/Engine/Source/Runtime/Engine/Classes/Animation/AttributeBlendOperator.h
+++ b/Engine/Source/Runtime/Engine/Classes/Animation/AttributeBlendOperator.h
@@ -22,6 +22,7 @@ namespace UE
 		public:
 			/** Begin IAttributeBlendOperator overrides */
 			virtual void Accumulate(const FAttributeBlendData& BlendData, FStackAttributeContainer* OutAttributes) const final { Accumulate_Internal<AttributeType>(BlendData, OutAttributes); }
+			virtual void AccumulateWithOverrideWeights(const FAttributeBlendData& BlendData, FStackAttributeContainer* OutAttributes) const final { AccumulateWithOverrideWeights_Internal<AttributeType>(BlendData, OutAttributes); }
 			virtual void Interpolate(const void* FromData, const void* ToData, float Alpha, void* InOutData) const final { Interpolate_Internal<AttributeType>(FromData, ToData, Alpha, InOutData); }
 			virtual void Override(const FAttributeBlendData& BlendData, FStackAttributeContainer* OutAttributes) const final { Override_Internal<AttributeType>(BlendData, OutAttributes); }
 			virtual void Blend(const FAttributeBlendData& BlendData, FStackAttributeContainer* OutAttributes) const final { Blend_Internal<AttributeType>(BlendData, OutAttributes); }
@@ -347,6 +348,73 @@ namespace UE
 					}
 				});
 			}
+			
+			/* Accumulate operation for blendable attribute types with optional per attribute weight overrides */
+			template <typename Type>
+			inline typename TEnableIf<TAttributeTypeTraits<Type>::IsBlendable, void>::Type AccumulateWithOverrideWeights_Internal(const FAttributeBlendData& BlendData, FStackAttributeContainer* OutAttributes) const
+			{
+				BlendData.ForEachUniqueAttribute<AttributeType>([OutAttributes](typename FAttributeBlendData::template TSingleIterator<AttributeType>& It) -> void
+				{
+					while (It.Next())
+					{
+						const AttributeType& Attribute = It.GetValue();
+						const FAttributeId& Identifier = It.GetIdentifier();
+						const int32 ExistingIndex = OutAttributes->IndexOfByKey<AttributeType>(Identifier);
+						AttributeType* OutAttributePtr = OutAttributes->FindOrAdd<AttributeType>(Identifier);
+						AttributeType& OutAttribute = *OutAttributePtr;
+						
+						float AttributeWeight = It.GetWeightOverride(Identifier);
+						if (AttributeWeight <= -1.0)
+						{
+							AttributeWeight = It.GetWeight();
+						}
+						
+						// Accumulate with weighted add
+						if (ExistingIndex != INDEX_NONE)
+						{						
+							OutAttribute.Accumulate(Attribute, AttributeWeight, It.GetAdditiveType());
+						}
+						else
+						{
+							OutAttribute = Attribute.Multiply(AttributeWeight);							
+						}
+
+						if constexpr (UE::Anim::TAttributeTypeTraits<AttributeType>::RequiresNormalization)
+						{
+							OutAttribute.Normalize();
+						}
+					}
+				});
+			}
+
+			/* Accumulate operation for non-blendable attribute types with optional per attribute weight overrides */
+			template <typename Type>
+			inline typename TEnableIf<!TAttributeTypeTraits<Type>::IsBlendable, void>::Type AccumulateWithOverrideWeights_Internal(const FAttributeBlendData& BlendData, FStackAttributeContainer* OutAttributes) const
+			{
+				BlendData.ForEachUniqueAttribute<AttributeType>([OutAttributes](typename FAttributeBlendData::template TSingleIterator<AttributeType>& It) -> void
+				{
+					while (It.Next())
+					{						
+						const AttributeType& Attribute = It.GetValue();
+						const FAttributeId& Identifier = It.GetIdentifier();
+						const int32 ExistingIndex = OutAttributes->IndexOfByKey<AttributeType>(Identifier);
+						AttributeType* OutAttributePtr = OutAttributes->FindOrAdd<AttributeType>(Identifier);
+						AttributeType& OutAttribute = *OutAttributePtr;
+						
+						float AttributeWeight = It.GetWeightOverride(Identifier);
+						if (AttributeWeight <= -1.0)
+						{
+							AttributeWeight = It.GetWeight();
+						}
+						
+						// 'Accumulate' value if the attribute did not yet exist, or based upon being the highest weighted attribute
+						if (ExistingIndex == INDEX_NONE || AttributeWeight > 0.5f)
+						{
+							AttributeType::StaticStruct()->CopyScriptStruct(OutAttributePtr, &Attribute);
+						}
+					}
+				});
+			}
 
 			/* Make additive operation for blendable attribute types */
 			template <typename Type>
@@ -450,6 +518,33 @@ namespace UE
 					}
 				});
 			};
+			
+			virtual void AccumulateWithOverrideWeights(const FAttributeBlendData& BlendData, FStackAttributeContainer* OutAttributes) const final
+			{
+				BlendData.ForEachUniqueAttribute([&, OutAttributes](FAttributeBlendData::TSingleRawIterator& It) -> void
+				{
+					while (It.Next())
+					{						
+						const uint8* Attribute= It.GetValuePtr();
+						const FAttributeId& Identifier = It.GetIdentifier();
+						const int32 ExistingIndex = OutAttributes->IndexOfByKey(ScriptStructPtr.Get(), Identifier);
+						uint8* OutAttributePtr = OutAttributes->FindOrAdd(ScriptStructPtr.Get(), Identifier);
+
+						float AttributeWeight = It.GetWeightOverride(Identifier);
+						if (AttributeWeight <= -1.0)
+						{
+							AttributeWeight = It.GetWeight();
+						}
+						
+						// 'Accumulate' value if the attribute did not yet exist, or based upon being the highest weighted attribute
+						if (ExistingIndex == INDEX_NONE || AttributeWeight > 0.5f)
+						{
+							ScriptStructPtr.Get()->CopyScriptStruct(OutAttributePtr, Attribute);
+						}
+					}
+				});
+			};
+			
 			virtual void Interpolate(const void* FromData, const void* ToData, float Alpha, void* InOutData) const final
 			{
 				// Determine stepped value to copy
diff --git a/Engine/Source/Runtime/Engine/Classes/Animation/AttributesContainer.h b/Engine/Source/Runtime/Engine/Classes/Animation/AttributesContainer.h
--- a/Engine/Source/Runtime/Engine/Classes/Animation/AttributesContainer.h
+++ b/Engine/Source/Runtime/Engine/Classes/Animation/AttributesContainer.h
@@ -33,6 +33,11 @@ namespace UE
 			{
 				return Index == Other.Index && Name == Other.Name && Namespace == Other.Namespace;
 			}
+			
+			friend inline uint32 GetTypeHash(const FAttributeId& AttributeId)
+			{
+				return HashCombine(GetTypeHash(AttributeId.Index), GetTypeHash(AttributeId.Name), GetTypeHash(AttributeId.Namespace));
+			}
 
 			int32 GetIndex() const { return Index; }
 			const FName& GetName() const { return Name; }
diff --git a/Engine/Source/Runtime/Engine/Classes/Animation/IAttributeBlendOperator.h b/Engine/Source/Runtime/Engine/Classes/Animation/IAttributeBlendOperator.h
--- a/Engine/Source/Runtime/Engine/Classes/Animation/IAttributeBlendOperator.h
+++ b/Engine/Source/Runtime/Engine/Classes/Animation/IAttributeBlendOperator.h
@@ -28,6 +28,9 @@ namespace UE
 
 			/** Invoked when an attribute container A is accumulated into container B */
 			virtual void Accumulate(const FAttributeBlendData& BlendData, FStackAttributeContainer* OutAttributes) const = 0;
+			
+			/** Invoked when an attribute container A is accumulated into container B with additional per attribute weight overrides */
+			virtual void AccumulateWithOverrideWeights(const FAttributeBlendData& BlendData, FStackAttributeContainer* OutAttributes) const = 0;
 
 			/** Invoked when an attribute container is supposed to be made additive with regards to container B */
 			virtual void ConvertToAdditive(const FAttributeBlendData& BlendData, FStackAttributeContainer* OutAdditiveAttributes) const = 0;
diff --git a/Engine/Source/Runtime/Engine/Private/Animation/AttributeBlendData.cpp b/Engine/Source/Runtime/Engine/Private/Animation/AttributeBlendData.cpp
--- a/Engine/Source/Runtime/Engine/Private/Animation/AttributeBlendData.cpp
+++ b/Engine/Source/Runtime/Engine/Private/Animation/AttributeBlendData.cpp
@@ -46,6 +46,17 @@ namespace UE { namespace Anim {
 			ProcessAttributes(CustomAttributes, SourceAttributesIndex, AttributeScriptStruct);
 		}
 	}
+	
+	FAttributeBlendData::FAttributeBlendData(const TArrayView<const FStackAttributeContainer* const> SourceAttributes, const TMap<UE::Anim::FAttributeId, float>& InAttributeWeightOverrides,  const float InUniformWeight, EAdditiveAnimationType InAdditiveType, const UScriptStruct* AttributeScriptStruct) : UniformWeight(InUniformWeight), AdditiveType(InAdditiveType)
+	{
+		AttributeWeightOverrides = InAttributeWeightOverrides;
+	
+		for (int32 SourceAttributesIndex = 0; SourceAttributesIndex < SourceAttributes.Num(); ++SourceAttributesIndex)
+		{
+			const FStackAttributeContainer& CustomAttributes = *SourceAttributes[SourceAttributesIndex];
+			ProcessAttributes(CustomAttributes, SourceAttributesIndex, AttributeScriptStruct);
+		}
+	}
 
 	FAttributeBlendData::FAttributeBlendData(const TArrayView<const FStackAttributeContainer> SourceAttributes, const TArrayView<const FPerBoneBlendWeight> InPerBoneBlendWeights, const UScriptStruct* AttributeScriptStruct) : UniformWeight(0.f), AdditiveType(AAT_None)
 	{
diff --git a/Engine/Source/Runtime/Engine/Private/Animation/AttributesRuntime.cpp b/Engine/Source/Runtime/Engine/Private/Animation/AttributesRuntime.cpp
--- a/Engine/Source/Runtime/Engine/Private/Animation/AttributesRuntime.cpp
+++ b/Engine/Source/Runtime/Engine/Private/Animation/AttributesRuntime.cpp
@@ -131,6 +131,24 @@ void Attributes::AccumulateAttributes(const FStackAttributeContainer& SourceAttr
 		Operator->Accumulate(AttributeBlendData, &OutAttributes);
 	}	
 }
+	
+void Attributes::AccumulateAttributesWithOverrideWeights(const FStackAttributeContainer& SourceAttributes, FStackAttributeContainer& OutAttributes, const TMap<FAttributeId, float>& AttributeWeightOverrides, float Weight, EAdditiveAnimationType AdditiveType)
+{
+	QUICK_SCOPE_CYCLE_COUNTER(STAT_AccumulateAttributesWithOverrideWeights);
+
+	TSet<TWeakObjectPtr<UScriptStruct>, DefaultKeyFuncs<TWeakObjectPtr<UScriptStruct>>, TInlineSetAllocator<4>> UniqueStructs;
+	UniqueStructs.Append(SourceAttributes.GetUniqueTypes());
+	UniqueStructs.Append(OutAttributes.GetUniqueTypes());
+
+	const FStackAttributeContainer* Array[] = { &SourceAttributes };
+	for (const TWeakObjectPtr<UScriptStruct>& WeakScriptStruct : UniqueStructs)
+	{
+		const IAttributeBlendOperator* Operator = AttributeTypes::GetTypeOperator(WeakScriptStruct);
+		FAttributeBlendData AttributeBlendData = FAttributeBlendData::SingleAdditiveContainerWithOverrideWeights(Array, AttributeWeightOverrides, Weight, AdditiveType, WeakScriptStruct.Get());
+		Operator->AccumulateWithOverrideWeights(AttributeBlendData, &OutAttributes);
+	}	
+}
+
 
 void Attributes::ConvertToAdditive(const FStackAttributeContainer& BaseAttributes, FStackAttributeContainer& OutAdditiveAttributes)
 {
diff --git a/Engine/Source/Runtime/Engine/Public/Animation/AttributesRuntime.h b/Engine/Source/Runtime/Engine/Public/Animation/AttributesRuntime.h
--- a/Engine/Source/Runtime/Engine/Public/Animation/AttributesRuntime.h
+++ b/Engine/Source/Runtime/Engine/Public/Animation/AttributesRuntime.h
@@ -88,6 +88,9 @@ namespace UE
 
 			/** Add any new or accumulate with existing custom attributes */
 			static ENGINE_API void AccumulateAttributes(const FStackAttributeContainer& SourceAttributes, FStackAttributeContainer& OutAttributes, float Weight, EAdditiveAnimationType AdditiveType);
+			
+			/** Add any new or accumulate with existing custom attributes with optional per attribute weight overrides */
+			static ENGINE_API void AccumulateAttributesWithOverrideWeights(const FStackAttributeContainer& SourceAttributes, FStackAttributeContainer& OutAttributes, const TMap<FAttributeId, float>& AttributeWeightOverrides, float Weight, EAdditiveAnimationType AdditiveType);
 
 			/** Add (negated) any new or subtract from existing custom attributes */
 			static ENGINE_API void ConvertToAdditive(const FStackAttributeContainer& BaseAttributes, FStackAttributeContainer& OutAdditiveAttributes);
