Properly account for looping sync markers in UAF

#Synced-CL 50252154

diff --git a/Engine/Plugins/Experimental/UAF/UAFAnimGraph/Source/UAFAnimGraph/Private/Graph/SyncGroup_GraphInstanceComponent.cpp b/Engine/Plugins/Experimental/UAF/UAFAnimGraph/Source/UAFAnimGraph/Private/Graph/SyncGroup_GraphInstanceComponent.cpp
--- a/Engine/Plugins/Experimental/UAF/UAFAnimGraph/Source/UAFAnimGraph/Private/Graph/SyncGroup_GraphInstanceComponent.cpp
+++ b/Engine/Plugins/Experimental/UAF/UAFAnimGraph/Source/UAFAnimGraph/Private/Graph/SyncGroup_GraphInstanceComponent.cpp
@@ -884,20 +884,22 @@ namespace UE::UAF
 				PhasePosition.NextMarkerName = NAME_None;
 			}
 
-			// Account for looping
 			if (PrevTime > NextTime)
 			{
+				// We looped, account for it
 				PrevTime = PrevTime > CurrentTime ? (PrevTime - Duration) : PrevTime;
 				NextTime = NextTime < CurrentTime ? (NextTime + Duration) : NextTime;
 			}
-			else if (PrevTime > CurrentTime)
+			else if (PrevTime == NextTime)
 			{
-				CurrentTime += Duration;
+				// If the two markers overlap, space them apart by the duration to ensure the time between them isn't zero
+				PrevTime -= Duration;
 			}
 
-			if (PrevTime == NextTime)
+			if (PrevTime > CurrentTime)
 			{
-				PrevTime -= Duration;
+				// Our current time isn't between the markers, we looped
+				CurrentTime += Duration;
 			}
 
 			ensure(NextTime > PrevTime);
