UAF: Add IGarbageCollection interface to the ApplyAdditivePerBone Trait

#Synced-CL 50297267

diff --git a/Engine/Plugins/Experimental/UAF/UAFAnimGraph/Source/UAFAnimGraph/Private/Traits/ApplyAdditivePerBoneTrait.cpp b/Engine/Plugins/Experimental/UAF/UAFAnimGraph/Source/UAFAnimGraph/Private/Traits/ApplyAdditivePerBoneTrait.cpp
--- a/Engine/Plugins/Experimental/UAF/UAFAnimGraph/Source/UAFAnimGraph/Private/Traits/ApplyAdditivePerBoneTrait.cpp
+++ b/Engine/Plugins/Experimental/UAF/UAFAnimGraph/Source/UAFAnimGraph/Private/Traits/ApplyAdditivePerBoneTrait.cpp
@@ -20,6 +20,7 @@ GeneratorMacro(IEvaluate) \
 GeneratorMacro(IHierarchy) \
 GeneratorMacro(IUpdate) \
 GeneratorMacro(IUpdateTraversal) \
+GeneratorMacro(IGarbageCollection) \
 
 GENERATE_ANIM_TRAIT_IMPLEMENTATION(FApplyAdditivePerBoneTrait, TRAIT_INTERFACE_ENUMERATOR, NULL_ANIM_TRAIT_INTERFACE_ENUMERATOR, NULL_ANIM_TRAIT_EVENT_ENUMERATOR)
 #undef TRAIT_INTERFACE_ENUMERATOR
@@ -163,4 +164,16 @@ float FApplyAdditivePerBoneTrait::GetBlendWeight(const FExecutionContext& Contex
 	// Invalid child index
 	return -1.0f;
 }
+
+void FApplyAdditivePerBoneTrait::AddReferencedObjects(const FExecutionContext& Context, const TTraitBinding<IGarbageCollection>& Binding, FReferenceCollector& Collector) const
+{
+	IGarbageCollection::AddReferencedObjects(Context, Binding, Collector);
+	
+	FInstanceData* InstanceData = Binding.GetInstanceData<FInstanceData>();
+	if (InstanceData->CachedBlendProfile.IsValid())
+	{
+		Collector.AddReferencedObject(InstanceData->CachedBlendProfile);
+	}
+}
+	
 }
diff --git a/Engine/Plugins/Experimental/UAF/UAFAnimGraph/Source/UAFAnimGraph/Private/Traits/ApplyAdditivePerBoneTrait.h b/Engine/Plugins/Experimental/UAF/UAFAnimGraph/Source/UAFAnimGraph/Private/Traits/ApplyAdditivePerBoneTrait.h
--- a/Engine/Plugins/Experimental/UAF/UAFAnimGraph/Source/UAFAnimGraph/Private/Traits/ApplyAdditivePerBoneTrait.h
+++ b/Engine/Plugins/Experimental/UAF/UAFAnimGraph/Source/UAFAnimGraph/Private/Traits/ApplyAdditivePerBoneTrait.h
@@ -5,6 +5,7 @@
 #include "TraitCore/Trait.h"
 #include "TraitInterfaces/IContinuousBlend.h"
 #include "TraitInterfaces/IEvaluate.h"
+#include "TraitInterfaces/IGarbageCollection.h"
 #include "TraitInterfaces/IHierarchy.h"
 #include "TraitInterfaces/IUpdate.h"
 #include "Traits/ApplyAdditivePerBoneTraitData.h"
@@ -17,10 +18,10 @@ namespace UE::UAF
 	 * 
 	 * A trait that can apply an additive animation with a hierarchy table based mask
 	 */
-	struct FApplyAdditivePerBoneTrait : FBaseTrait, IEvaluate, IUpdate, IUpdateTraversal, IHierarchy, IContinuousBlend
+	struct FApplyAdditivePerBoneTrait : FBaseTrait, IEvaluate, IUpdate, IUpdateTraversal, IHierarchy, IContinuousBlend, IGarbageCollection
 	{
 		DECLARE_ANIM_TRAIT(FApplyAdditivePerBoneTrait, FBaseTrait)
-		
+
 		using FSharedData = FApplyAdditivePerBoneTraitSharedData;
 
 		struct FInstanceData : FTrait::FInstanceData
@@ -51,5 +52,8 @@ namespace UE::UAF
 
 		// IContinuousBlend impl
 		virtual float GetBlendWeight(const FExecutionContext& Context, const TTraitBinding<IContinuousBlend>& Binding, int32 ChildIndex) const override;
+		
+		// IGarbageCollection impl
+		virtual void AddReferencedObjects(const FExecutionContext& Context, const TTraitBinding<IGarbageCollection>& Binding, FReferenceCollector& Collector) const override;
 	};
 }
