- Add ability to retrieve compilation state for UAF asset
- Added missing python output for UAnimNextController::AddTraitStruct

#Synced-CL 50296236

diff --git a/Engine/Plugins/Experimental/UAF/UAF/Source/UAF/Public/AnimNextRigVMAsset.h b/Engine/Plugins/Experimental/UAF/UAF/Source/UAF/Public/AnimNextRigVMAsset.h
--- a/Engine/Plugins/Experimental/UAF/UAF/Source/UAF/Public/AnimNextRigVMAsset.h
+++ b/Engine/Plugins/Experimental/UAF/UAF/Source/UAF/Public/AnimNextRigVMAsset.h
@@ -46,12 +46,13 @@ namespace UE::UAF::UncookedOnly
 }
 
 #if WITH_EDITORONLY_DATA
-UENUM()
+UENUM(BlueprintType)
 enum class EAnimNextRigVMAssetState : uint8
 {
 	CompiledWithErrors,
 	CompiledWithWarnings,
-	CompiledWithSuccess
+	CompiledWithSuccess,
+	Invalid
 };
 #endif // WITH_EDITORONLY_DATA
 
diff --git a/Engine/Plugins/Experimental/UAF/UAFAnimGraph/Source/UAFAnimGraphUncookedOnly/Private/AnimNextController.cpp b/Engine/Plugins/Experimental/UAF/UAFAnimGraph/Source/UAFAnimGraphUncookedOnly/Private/AnimNextController.cpp
--- a/Engine/Plugins/Experimental/UAF/UAFAnimGraph/Source/UAFAnimGraphUncookedOnly/Private/AnimNextController.cpp
+++ b/Engine/Plugins/Experimental/UAF/UAFAnimGraph/Source/UAFAnimGraphUncookedOnly/Private/AnimNextController.cpp
@@ -64,7 +64,7 @@ FName UAnimNextController::AddTraitStruct(URigVMUnitNode* InNode, const TInstanc
 	// Must use struct name for name to allow for consistent pathing in node layouts
 	const FName TraitName = InTraitDefaults.GetScriptStruct()->GetFName();
 
-	FName NewTraitName = Super::AddTrait(InNode, CppDecoratorStruct, TraitName, DefaultValue, InTraitIndex, bSetupUndoRedo);
+	const FName NewTraitName = Super::AddTrait(InNode, CppDecoratorStruct, TraitName, DefaultValue, InTraitIndex, bSetupUndoRedo);
 	if (NewTraitName != NAME_None)
 	{
 		// Append defaults for shared data
@@ -92,8 +92,17 @@ FName UAnimNextController::AddTraitStruct(URigVMUnitNode* InNode, const TInstanc
 			}
 		}
 	}
+	
+	if(!NewTraitName.IsNone() && bPrintPythonCommand)
+	{
+		const TArray<FString> TraitCommands = GetAddTraitPythonCommands(InNode, NewTraitName);
+		for(const FString& TraitCommand : TraitCommands)
+		{
+			RigVMPythonUtils::Print(GetSchema()->GetGraphOuterName(GetGraph()), TraitCommand);
+		}
+	}
 
-	return TraitName;
+	return NewTraitName;
 }
 
 FName UAnimNextController::AddTraitByName(FName InNodeName, FName InNewTraitTypeName, int32 InPinIndex, const FString& InNewTraitDefaultValue, bool bSetupUndoRedo, bool bPrintPythonCommand)
diff --git a/Engine/Plugins/Experimental/UAF/UAFTests/Source/UAFTests/Private/UAFTestBlueprintLibrary.h b/Engine/Plugins/Experimental/UAF/UAFTests/Source/UAFTests/Private/UAFTestBlueprintLibrary.h
--- a/Engine/Plugins/Experimental/UAF/UAFTests/Source/UAFTests/Private/UAFTestBlueprintLibrary.h
+++ b/Engine/Plugins/Experimental/UAF/UAFTests/Source/UAFTests/Private/UAFTestBlueprintLibrary.h
@@ -260,4 +260,17 @@ public:
 		
 		return Asset;
 	}
+	
+	UFUNCTION(BlueprintPure, Category="UAFTest", meta=(ScriptMethod))
+	static EAnimNextRigVMAssetState GetCompilationState(const UUAFRigVMAsset* InAsset)
+	{
+#if WITH_EDITORONLY_DATA
+		if (InAsset)
+		{
+			return InAsset->CompilationState;
+		}
+#endif // WITH_EDITORONLY_DATA
+		
+		return EAnimNextRigVMAssetState::Invalid;
+	}
 };
