UAF: Add Make Dynamic Additive Trait

- Adds a new trait to create a dynamic additive based on the two input poses
- Adds a task to create a dynamic additive for skeleton and value runtime

#Synced-CL 50345329

diff --git a/Engine/Plugins/Experimental/UAF/UAF/Source/UAF/Private/TransformArrayOperations_AoS.cpp b/Engine/Plugins/Experimental/UAF/UAF/Source/UAF/Private/TransformArrayOperations_AoS.cpp
--- a/Engine/Plugins/Experimental/UAF/UAF/Source/UAF/Private/TransformArrayOperations_AoS.cpp
+++ b/Engine/Plugins/Experimental/UAF/UAF/Source/UAF/Private/TransformArrayOperations_AoS.cpp
@@ -232,4 +232,28 @@ namespace UE::UAF
 			Dest[LODBoneIndex].AccumulateWithShortestRotation(Source[LODBoneIndex], VScaleWeight);
 		}
 	}
+	
+	void ConvertToAdditive(const FTransformArrayAoSConstView& Input, const FTransformArrayAoSConstView& Base, const FTransformArrayAoSView& Output)
+	{
+		const int32 NumTransforms = Input.Num();
+		check(Base.Num() == NumTransforms);
+		check(Output.Num() == NumTransforms);
+		
+		for (int32 TransformIndex = 0; TransformIndex < NumTransforms; ++TransformIndex)
+		{
+			const FTransform& InputTransform = Input[TransformIndex];
+			const FTransform& BaseTransform = Base[TransformIndex];
+			FTransform& OutTransform = Output[TransformIndex];
+			
+			OutTransform.SetRotation(InputTransform.GetRotation() * BaseTransform.GetRotation().Inverse());
+			OutTransform.SetTranslation(InputTransform.GetTranslation() - BaseTransform.GetTranslation());
+			// additive scale considers how much it grow or lower
+			// in order to support blending between different additive scale, we save [(target scale)/(source scale) - 1.f], and this can blend with 
+			// other delta scale value
+			// when we apply to the another scale, we apply scale * (1 + [additive scale])
+			OutTransform.SetScale3D(InputTransform.GetScale3D() * FTransform::GetSafeScaleReciprocal(BaseTransform.GetScale3D()) -1.0f); 
+			
+			OutTransform.NormalizeRotation();
+		}
+	}
 }
diff --git a/Engine/Plugins/Experimental/UAF/UAF/Source/UAF/Private/TransformArrayOperations_SoA.cpp b/Engine/Plugins/Experimental/UAF/UAF/Source/UAF/Private/TransformArrayOperations_SoA.cpp
--- a/Engine/Plugins/Experimental/UAF/UAF/Source/UAF/Private/TransformArrayOperations_SoA.cpp
+++ b/Engine/Plugins/Experimental/UAF/UAF/Source/UAF/Private/TransformArrayOperations_SoA.cpp
@@ -449,4 +449,24 @@ namespace UE::UAF
 			Dest.Scales3D[LODBoneIndex] += Source.Scales3D[LODBoneIndex] * ScaleWeight;
 		}
 	}
+	
+	void ConvertToAdditive(const FTransformArraySoAConstView& Input, const FTransformArraySoAConstView& Base, const FTransformArraySoAView& Output)
+	{
+		const int32 NumTransforms = Input.Num();
+		check(Base.Num() == NumTransforms);
+		check(Output.Num() == NumTransforms);
+		
+		for (int32 TransformIndex = 0; TransformIndex < NumTransforms; ++TransformIndex)
+		{
+			Output.Rotations[TransformIndex] = Input[TransformIndex].GetRotation() * Base[TransformIndex].GetRotation().Inverse();
+			Output.Translations[TransformIndex] = Input[TransformIndex].GetTranslation() - Base[TransformIndex].GetTranslation();
+			// additive scale considers how much it grow or lower
+			// in order to support blending between different additive scale, we save [(target scale)/(source scale) - 1.f], and this can blend with 
+			// other delta scale value
+			// when we apply to the another scale, we apply scale * (1 + [additive scale])
+			Output.Scales3D[TransformIndex] = Input[TransformIndex].GetScale3D() * FTransform::GetSafeScaleReciprocal(Base[TransformIndex].GetScale3D()) -1.0f; 
+			
+			Output.Rotations[TransformIndex].Normalize();
+		}
+	}
 }
diff --git a/Engine/Plugins/Experimental/UAF/UAF/Source/UAF/Public/TransformArrayOperations.h b/Engine/Plugins/Experimental/UAF/UAF/Source/UAF/Public/TransformArrayOperations.h
--- a/Engine/Plugins/Experimental/UAF/UAF/Source/UAF/Public/TransformArrayOperations.h
+++ b/Engine/Plugins/Experimental/UAF/UAF/Source/UAF/Public/TransformArrayOperations.h
@@ -98,4 +98,9 @@ namespace UE::UAF
 	UAF_API void BlendAddPerBoneWithScale(
 		const FTransformArraySoAView& Dest, const FTransformArraySoAConstView& Source,
 		const TArrayView<const int32>& LODBoneIndexToWeightIndexMap, const TArrayView<const float>& BoneWeights, const float DefaultScaleWeight);
+	
+	// Convert Input into an Additive Pose, by doing Output = Input - Base 
+	// @TODO: AoS version is untested because AoS currently does not compile in UAF.
+	UAF_API void ConvertToAdditive(const FTransformArrayAoSConstView& Input, const FTransformArrayAoSConstView& Base, const FTransformArrayAoSView& Output);
+	UAF_API void ConvertToAdditive(const FTransformArraySoAConstView& Input, const FTransformArraySoAConstView& Base, const FTransformArraySoAView& Output);
 }
diff --git a/Engine/Plugins/Experimental/UAF/UAFAnimGraph/Source/UAFAnimGraph/Private/EvaluationVM/Tasks/MakeDynamicAdditive.cpp b/Engine/Plugins/Experimental/UAF/UAFAnimGraph/Source/UAFAnimGraph/Private/EvaluationVM/Tasks/MakeDynamicAdditive.cpp
new file mode 100644
--- /dev/null
+++ b/Engine/Plugins/Experimental/UAF/UAFAnimGraph/Source/UAFAnimGraph/Private/EvaluationVM/Tasks/MakeDynamicAdditive.cpp
@@ -0,0 +1,106 @@
+// Copyright Epic Games, Inc. All Rights Reserved.
+
+#include "EvaluationVM/Tasks/MakeDynamicAdditive.h"
+
+#include "EvaluationVM/EvaluationVM.h"
+#include "UAF/ValueRuntime/Transformers/AdditiveSpace.h"
+
+#include UE_INLINE_GENERATED_CPP_BY_NAME(MakeDynamicAdditive)
+
+FUAFMakeDynamicAdditiveTask FUAFMakeDynamicAdditiveTask::Make(const EAdditiveAnimationType InAdditiveType)
+{
+	FUAFMakeDynamicAdditiveTask Task;
+	Task.AdditiveType = InAdditiveType;
+	return Task;
+}
+
+void FUAFMakeDynamicAdditiveTask::Execute(UE::UAF::FEvaluationVM& VM) const
+{
+	using namespace UE::UAF;
+	
+	if (VM.GetActiveNamedSet())
+	{
+		// Pop our top two poses, the top one being the pose to be turned additive and the second one the base
+		TUniquePtr<FValueBundle> AdditiveInput;
+		if (!VM.PopValue(ATTRIBUTE_STACK_NAME, AdditiveInput))
+		{
+			// We have no inputs, nothing to do
+			return;
+		}
+
+		FValueSpace ValueSpace;
+		if (AdditiveType == AAT_RotationOffsetMeshSpace)
+		{
+			ValueSpace = FValueSpace(EMixedSpaceFlags::MeshRotation, true);
+		}
+		else
+		{
+			ValueSpace = FValueSpace(EValueSpaceType::Local, true);
+		}
+		
+		FValueBundleStack Output(VM.GetActiveNamedSet());
+		Output.InitWithValueSpace(ValueSpace);
+		
+		TUniquePtr<FValueBundle> BaseInput;
+		if (!VM.PopValue(ATTRIBUTE_STACK_NAME, BaseInput))
+		{
+			// If we have no base to generate against, we push an additive ref pose in the correct value space 
+			VM.PushValue(ATTRIBUTE_STACK_NAME, MakeUnique<FValueBundle>(MoveTemp(Output)));
+			return;
+		}
+		
+		UE::UAF::Transformers::FMakeAdditiveSpace::Apply(VM.GetTransformerMap(), *BaseInput, *AdditiveInput, Output);
+		VM.PushValue(ATTRIBUTE_STACK_NAME, MakeUnique<FValueBundle>(MoveTemp(Output)));
+	}
+	else
+	{
+		// Pop our top two poses, the top one being the pose to be turned additive and the second one the base
+		TUniquePtr<FKeyframeState> AdditiveKeyframe;
+		if (!VM.PopValue(KEYFRAME_STACK_NAME, AdditiveKeyframe))
+		{
+			// We have no inputs, nothing to do
+			return;
+		}
+
+		TUniquePtr<FKeyframeState> BaseKeyframe;
+		if (!VM.PopValue(KEYFRAME_STACK_NAME, BaseKeyframe))
+		{
+			// If we have no base to generate against, we push an additive ref pose 
+			VM.PushValue(KEYFRAME_STACK_NAME, MakeUnique<FKeyframeState>(VM.MakeReferenceKeyframe(true)));
+			return;
+		}
+		
+		FKeyframeState OutputKeyframe = VM.MakeUninitializedKeyframe(true);
+		if (EnumHasAnyFlags(VM.GetFlags(), EEvaluationFlags::Bones))
+		{
+			// Build bone data 
+			ConvertToAdditive(AdditiveKeyframe->Pose.LocalTransforms.GetConstView(), BaseKeyframe->Pose.LocalTransforms.GetConstView(), OutputKeyframe.Pose.LocalTransforms.GetView());
+		}
+
+		if (EnumHasAnyFlags(VM.GetFlags(), EEvaluationFlags::Curves))
+		{
+			// Build curve data
+			// Convert current curves to Additive (Target - BaseCurve) if overlapping entries are found
+			OutputKeyframe.Curves.CopyFrom(AdditiveKeyframe->Curves);
+			FBlendedCurve& BaseCurve = BaseKeyframe->Curves;
+			FBlendedCurve& AdditiveCurve = OutputKeyframe.Curves;
+			UE::Anim::FNamedValueArrayUtils::Union(
+				AdditiveCurve, 
+				BaseCurve,
+				[](UE::Anim::FCurveElement& InOutTargetCurveElement, const UE::Anim::FCurveElement& InBaseCurveElement, UE::Anim::ENamedValueUnionFlags InFlags)
+			{
+				InOutTargetCurveElement.Value -= InBaseCurveElement.Value;
+				InOutTargetCurveElement.Flags |= InBaseCurveElement.Flags;
+			});
+		}
+
+		if (EnumHasAnyFlags(VM.GetFlags(), EEvaluationFlags::Attributes))
+		{
+			// Build attribute data 
+			OutputKeyframe.Attributes.CopyFrom(AdditiveKeyframe->Attributes);
+			UE::Anim::Attributes::ConvertToAdditive(BaseKeyframe->Attributes, OutputKeyframe.Attributes);
+		}
+		
+		VM.PushValue(KEYFRAME_STACK_NAME, MakeUnique<FKeyframeState>(MoveTemp(OutputKeyframe)));
+	}
+}
diff --git a/Engine/Plugins/Experimental/UAF/UAFAnimGraph/Source/UAFAnimGraph/Private/Traits/MakeDynamicAdditiveTrait.cpp b/Engine/Plugins/Experimental/UAF/UAFAnimGraph/Source/UAFAnimGraph/Private/Traits/MakeDynamicAdditiveTrait.cpp
new file mode 100644
--- /dev/null
+++ b/Engine/Plugins/Experimental/UAF/UAFAnimGraph/Source/UAFAnimGraph/Private/Traits/MakeDynamicAdditiveTrait.cpp
@@ -0,0 +1,67 @@
+// Copyright Epic Games, Inc. All Rights Reserved.
+
+#include "Traits/MakeDynamicAdditiveTrait.h"
+
+#include "EvaluationVM/Tasks/ConvertRotationsToMeshSpace.h"
+#include "EvaluationVM/Tasks/MakeDynamicAdditive.h"
+
+namespace UE::UAF
+{
+	AUTO_REGISTER_ANIM_TRAIT(FMakeDynamicAdditiveTrait)
+
+	#define TRAIT_INTERFACE_ENUMERATOR(GeneratorMacro) \
+		GeneratorMacro(IEvaluate) \
+		GeneratorMacro(IHierarchy) \
+
+	GENERATE_ANIM_TRAIT_IMPLEMENTATION(FMakeDynamicAdditiveTrait, TRAIT_INTERFACE_ENUMERATOR, NULL_ANIM_TRAIT_INTERFACE_ENUMERATOR, NULL_ANIM_TRAIT_EVENT_ENUMERATOR)
+	#undef TRAIT_INTERFACE_ENUMERATOR
+
+	void FMakeDynamicAdditiveTrait::FInstanceData::Construct(const FExecutionContext& Context, const FTraitBinding& Binding)
+	{
+		FTrait::FInstanceData::Construct(Context, Binding);
+		
+		const FSharedData* SharedData = Binding.GetSharedData<FSharedData>();
+		FInstanceData* InstanceData = Binding.GetInstanceData<FInstanceData>();
+
+		if (!InstanceData->Base.IsValid())
+		{
+			InstanceData->Base = Context.AllocateNodeInstance(Binding, SharedData->Base);
+		}
+
+		if (!InstanceData->Additive.IsValid())
+		{
+			InstanceData->Additive = Context.AllocateNodeInstance(Binding, SharedData->Additive);
+		}
+	}
+
+	void FMakeDynamicAdditiveTrait::PostEvaluate(FEvaluateTraversalContext& Context, const TTraitBinding<IEvaluate>& Binding) const
+	{
+		const FSharedData* SharedData = Binding.GetSharedData<FSharedData>();
+		const FInstanceData* InstanceData = Binding.GetInstanceData<FInstanceData>();
+
+		if (InstanceData->Additive.IsValid() && InstanceData->Base.IsValid())
+		{
+			const bool bMeshSpaceAdditive = SharedData->GetbMeshSpaceAdditive(Binding);
+			if (bMeshSpaceAdditive)
+			{
+				Context.AppendTask(FAnimNextConvertRotationsToMeshSpaceTask::Make(2));
+			}
+			
+			Context.AppendTask(FUAFMakeDynamicAdditiveTask::Make(bMeshSpaceAdditive ? AAT_RotationOffsetMeshSpace : AAT_LocalSpaceBase));
+		}
+	}
+
+	uint32 FMakeDynamicAdditiveTrait::GetNumChildren(const FExecutionContext& Context, const TTraitBinding<IHierarchy>& Binding) const
+	{
+		return 2;
+	}
+
+	void FMakeDynamicAdditiveTrait::GetChildren(const FExecutionContext& Context, const TTraitBinding<IHierarchy>& Binding, FChildrenArray& Children) const
+	{
+		const FInstanceData* InstanceData = Binding.GetInstanceData<FInstanceData>();
+
+		// Add the two children, even if the handles are empty
+		Children.Add(InstanceData->Base);
+		Children.Add(InstanceData->Additive);
+	}
+}
diff --git a/Engine/Plugins/Experimental/UAF/UAFAnimGraph/Source/UAFAnimGraph/Private/Traits/MakeDynamicAdditiveTrait.h b/Engine/Plugins/Experimental/UAF/UAFAnimGraph/Source/UAFAnimGraph/Private/Traits/MakeDynamicAdditiveTrait.h
new file mode 100644
--- /dev/null
+++ b/Engine/Plugins/Experimental/UAF/UAFAnimGraph/Source/UAFAnimGraph/Private/Traits/MakeDynamicAdditiveTrait.h
@@ -0,0 +1,42 @@
+// Copyright Epic Games, Inc. All Rights Reserved.
+
+#pragma once
+
+#include "TraitCore/Trait.h"
+#include "TraitInterfaces/IEvaluate.h"
+#include "TraitInterfaces/IHierarchy.h"
+#include "TraitInterfaces/IUpdate.h"
+#include "Traits/MakeDynamicAdditiveTraitData.h"
+
+
+namespace UE::UAF
+{
+	/**
+	 * Make Dynamic Additive Trait
+	 * 
+	 * A trait that can generate an additive pose based on the input poses
+	 * The Additive input is the pose to be turned additive 
+	 * The Base input is the base to generate the additive against
+	 */
+	struct FMakeDynamicAdditiveTrait : FBaseTrait, IEvaluate, IHierarchy
+	{
+		DECLARE_ANIM_TRAIT(FMakeDynamicAdditiveTrait, FBaseTrait)
+		
+		using FSharedData = FMakeDynamicAdditiveTraitSharedData;
+
+		struct FInstanceData : FTrait::FInstanceData
+		{
+			FTraitPtr Base;
+			FTraitPtr Additive;
+			
+			void Construct(const FExecutionContext& Context, const FTraitBinding& Binding);
+		};
+		
+		// IEvaluate impl
+		virtual void PostEvaluate(FEvaluateTraversalContext& Context, const TTraitBinding<IEvaluate>& Binding) const override;
+		
+		// IHierarchy impl
+		virtual uint32 GetNumChildren(const FExecutionContext& Context, const TTraitBinding<IHierarchy>& Binding) const override;
+		virtual void GetChildren(const FExecutionContext& Context, const TTraitBinding<IHierarchy>& Binding, FChildrenArray& Children) const override;
+	};
+}
diff --git a/Engine/Plugins/Experimental/UAF/UAFAnimGraph/Source/UAFAnimGraph/Public/EvaluationVM/Tasks/MakeDynamicAdditive.h b/Engine/Plugins/Experimental/UAF/UAFAnimGraph/Source/UAFAnimGraph/Public/EvaluationVM/Tasks/MakeDynamicAdditive.h
new file mode 100644
--- /dev/null
+++ b/Engine/Plugins/Experimental/UAF/UAFAnimGraph/Source/UAFAnimGraph/Public/EvaluationVM/Tasks/MakeDynamicAdditive.h
@@ -0,0 +1,30 @@
+// Copyright Epic Games, Inc. All Rights Reserved.
+
+#pragma once
+
+#include "EvaluationVM/EvaluationTask.h"
+#include "Animation/AttributeTypes.h"
+
+#include "MakeDynamicAdditive.generated.h"
+
+#define UE_API UAFANIMGRAPH_API
+
+/* Make Dynamic Additive Task
+ * 
+ * Pops the top two keyframes of the stack, creates a new additive pose and pushes it back onto the stack.
+ * The top pose should be the keyframe to turn into an additive and the second to the top should be the base keyframe. 
+*/
+USTRUCT()
+struct FUAFMakeDynamicAdditiveTask : public FAnimNextEvaluationTask
+{
+	GENERATED_BODY()
+
+	DECLARE_ANIM_EVALUATION_TASK(FUAFMakeDynamicAdditiveTask)
+
+	static UE_API FUAFMakeDynamicAdditiveTask Make(const EAdditiveAnimationType InAdditiveType);
+	UE_API virtual void Execute(UE::UAF::FEvaluationVM& VM) const override;
+	
+	EAdditiveAnimationType AdditiveType = AAT_LocalSpaceBase;
+};
+
+#undef UE_API
diff --git a/Engine/Plugins/Experimental/UAF/UAFAnimGraph/Source/UAFAnimGraph/Public/Traits/MakeDynamicAdditiveTraitData.h b/Engine/Plugins/Experimental/UAF/UAFAnimGraph/Source/UAFAnimGraph/Public/Traits/MakeDynamicAdditiveTraitData.h
new file mode 100644
--- /dev/null
+++ b/Engine/Plugins/Experimental/UAF/UAFAnimGraph/Source/UAFAnimGraph/Public/Traits/MakeDynamicAdditiveTraitData.h
@@ -0,0 +1,38 @@
+// Copyright Epic Games, Inc. All Rights Reserved.
+
+#pragma once
+
+#include "TraitCore/Trait.h"
+#include "TraitCore/TraitBinding.h"
+
+#include "MakeDynamicAdditiveTraitData.generated.h"
+
+
+namespace UE::UAF
+{
+	// A trait that creates a dynamic additive based input poses
+	USTRUCT(meta = (DisplayName = "Make Dynamic Additive", ShowTooltip=true))
+	struct FMakeDynamicAdditiveTraitSharedData : public FAnimNextTraitSharedData
+	{
+		GENERATED_BODY()
+
+		// Reference pose for additive delta calculation
+		UPROPERTY()
+		FAnimNextTraitHandle Base;
+
+		// Pose to make additive
+		UPROPERTY()
+		FAnimNextTraitHandle Additive;
+		
+		// Do additive delta calculation in mesh space
+		UPROPERTY(EditAnywhere, Category="Default")
+		bool bMeshSpaceAdditive;
+		
+		// Latent pin support boilerplate
+		#define TRAIT_LATENT_PROPERTIES_ENUMERATOR(GeneratorMacro) \
+		GeneratorMacro(bMeshSpaceAdditive) \
+
+		GENERATE_TRAIT_LATENT_PROPERTIES(FMakeDynamicAdditiveTraitSharedData, TRAIT_LATENT_PROPERTIES_ENUMERATOR)
+		#undef TRAIT_LATENT_PROPERTIES_ENUMERATOR
+	};
+}
\ No newline at end of file
