UAF: Pipe through bSetupUndoRedo to SetPinDefaultValue in UAnimNextController::AddTraitStruct

#Synced-CL 50394599

diff --git a/Engine/Plugins/Experimental/UAF/UAFAnimGraph/Source/UAFAnimGraphUncookedOnly/Private/AnimNextController.cpp b/Engine/Plugins/Experimental/UAF/UAFAnimGraph/Source/UAFAnimGraphUncookedOnly/Private/AnimNextController.cpp
--- a/Engine/Plugins/Experimental/UAF/UAFAnimGraph/Source/UAFAnimGraphUncookedOnly/Private/AnimNextController.cpp
+++ b/Engine/Plugins/Experimental/UAF/UAFAnimGraph/Source/UAFAnimGraphUncookedOnly/Private/AnimNextController.cpp
@@ -85,7 +85,7 @@ FName UAnimNextController::AddTraitStruct(URigVMUnitNode* InNode, const TInstanc
 						PostProcessDefaultValue(SubPin, MemberValue);
 						if (!MemberValue.IsEmpty())
 						{
-							SetPinDefaultValue(SubPin, MemberValue, true, true, true, true);
+							SetPinDefaultValue(SubPin, MemberValue, true, bSetupUndoRedo, true, true);
 						}
 					}
 				}
