Fix friend of a non-existent symbol

#rnx

diff --git a/Engine/Plugins/Experimental/UAF/UAFAnimGraph/Source/UAFAnimGraph/Internal/Factory/AnimGraphBuilderContext.h b/Engine/Plugins/Experimental/UAF/UAFAnimGraph/Source/UAFAnimGraph/Internal/Factory/AnimGraphBuilderContext.h
--- a/Engine/Plugins/Experimental/UAF/UAFAnimGraph/Source/UAFAnimGraph/Internal/Factory/AnimGraphBuilderContext.h
+++ b/Engine/Plugins/Experimental/UAF/UAFAnimGraph/Source/UAFAnimGraph/Internal/Factory/AnimGraphBuilderContext.h
@@ -21,7 +21,7 @@ struct FAnimGraphBuilderContext
 {
 private:
 	friend FAnimGraphFactory;
-	friend ::FAnimationAnimNextRuntimeTest_SimpleGraphBuilder;
+	friend FAnimationAnimNextRuntimeTest_SimpleGraphBuilder;
 
 	// Builds a graph using the populated TraitWriter and VariableStructs
 	UAFANIMGRAPH_API const UUAFAnimGraph* Build();
