Fix FBlendStackCoreTrait::GetOnExtractRootMotionAttribute to find sequence data using new trait data native struct

#Synced-CL 50467870

diff --git a/Engine/Plugins/Experimental/UAF/UAFAnimGraph/Source/UAFAnimGraph/Private/Traits/BlendStackTrait.cpp b/Engine/Plugins/Experimental/UAF/UAFAnimGraph/Source/UAFAnimGraph/Private/Traits/BlendStackTrait.cpp
--- a/Engine/Plugins/Experimental/UAF/UAFAnimGraph/Source/UAFAnimGraph/Private/Traits/BlendStackTrait.cpp
+++ b/Engine/Plugins/Experimental/UAF/UAFAnimGraph/Source/UAFAnimGraph/Private/Traits/BlendStackTrait.cpp
@@ -330,7 +330,7 @@ namespace UE::UAF
 			// @TODO: Selecting the first active graph in a blend stack doesn't typically make sense, a wrapper trait is better.
 			// However this implementation is still useful for prototyping / debugging purposes. Still consider future removal.
 			const UAnimSequence* AnimSequence = nullptr;
-			ActiveGraphInstance->AccessVariablesStruct<FAnimNextNativeDataInterface_AnimSequencePlayer>([&AnimSequence](FAnimNextNativeDataInterface_AnimSequencePlayer& InAnimSequencePlayer)
+			ActiveGraphInstance->AccessVariablesStruct<FAnimNextSequencePlayerTraitSharedData>([&AnimSequence](FAnimNextSequencePlayerTraitSharedData&InAnimSequencePlayer)
 			{
 				if (InAnimSequencePlayer.AnimSequence)
 				{
