UAF: Add blueprint function and RigUnit to be able to query if a montage slot is active in UAF

#Synced-CL 50531632

diff --git a/Engine/Plugins/Experimental/UAF/UAFAnimGraph/Source/UAFAnimGraph/Internal/Injection/UAFMontageComponent.h b/Engine/Plugins/Experimental/UAF/UAFAnimGraph/Source/UAFAnimGraph/Internal/Injection/UAFMontageComponent.h
--- a/Engine/Plugins/Experimental/UAF/UAFAnimGraph/Source/UAFAnimGraph/Internal/Injection/UAFMontageComponent.h
+++ b/Engine/Plugins/Experimental/UAF/UAFAnimGraph/Source/UAFAnimGraph/Internal/Injection/UAFMontageComponent.h
@@ -3,13 +3,13 @@
 #pragma once
 
 #include "CoreMinimal.h"
-#include "Graph/UAFGraphInstanceComponent.h"
-#include "Module/UAFModuleInstanceComponent.h"
-#include "Engine/EngineBaseTypes.h"
 #include "AlphaBlend.h"
-#include "Animation/AnimNotifyQueue.h"
-#include "Animation/AnimMontage.h"
 #include "MontageTrait.h"
+#include "Animation/AnimMontage.h"
+#include "Animation/AnimNotifyQueue.h"
+#include "Engine/EngineBaseTypes.h"
+#include "Graph/RigUnit_AnimNextBase.h"
+#include "Module/UAFModuleInstanceComponent.h"
 #include "UAFMontageComponent.generated.h"
 
 struct FUAFUpdateMontageTickFunction : public FTickFunction
@@ -45,7 +45,13 @@ public:
 		, BlendStartAlpha(InBlendStartAlpha)
 	{
 	}
-
+	
+	// Returns true if the montage in this eval data is active (valid and not blending out)
+	bool IsActive() const { return (Montage != nullptr && BlendInfo.GetDesiredValue() != 0.0f); }
+	
+	// Return the montage for this eval data
+	const UAnimMontage* GetMontage() const {return Montage;}
+	
 private:
 	UPROPERTY(Transient)
 	TObjectPtr<UAnimMontage> Montage = nullptr;
@@ -75,6 +81,7 @@ public:
 	FUAFMontageComponent() = default;
 	virtual ~FUAFMontageComponent() override;
 
+	bool IsSlotActive(FName SlotNodeName) const;
 private:
 	virtual void OnBindToInstance() override;
 
@@ -125,3 +132,22 @@ struct TStructOpsTypeTraits<FUAFMontageComponent> : public TStructOpsTypeTraitsB
 		WithCopy = false,
 	};
 };
+
+// Return true if the running UAF System has an active montage in the given slot. A UAnimMontage that is playing in the slot and blending out is not determined to be "active". 
+USTRUCT(meta=(DisplayName="Is Slot Active", Category="Injection", NodeColor="0, 1, 1", Keywords="Slot, Active, Montage"))
+struct FRigUnit_IsSlotActive : public FRigUnit_AnimNextBase
+{
+	GENERATED_BODY()
+
+	RIGVM_METHOD()
+	void Execute();
+	
+	// The name of the slot
+	UPROPERTY(EditAnywhere, Category = "Graph", meta = (Input))
+	FName SlotNodeName = NAME_None;
+	
+	// True if the running UAF System has an active montage in the given slot and it's not blending out
+	UPROPERTY(EditAnywhere, Category = "Graph", meta = (Output))
+	bool Result = false;
+};
+
diff --git a/Engine/Plugins/Experimental/UAF/UAFAnimGraph/Source/UAFAnimGraph/Private/Injection/InjectionLibrary.cpp b/Engine/Plugins/Experimental/UAF/UAFAnimGraph/Source/UAFAnimGraph/Private/Injection/InjectionLibrary.cpp
new file mode 100644
--- /dev/null
+++ b/Engine/Plugins/Experimental/UAF/UAFAnimGraph/Source/UAFAnimGraph/Private/Injection/InjectionLibrary.cpp
@@ -0,0 +1,39 @@
+// Copyright Epic Games, Inc. All Rights Reserved.
+
+#include "Injection/InjectionLibrary.h"
+
+#include "Animation/AnimMontage.h"
+#include "Component/AnimNextComponent.h"
+
+#include UE_INLINE_GENERATED_CPP_BY_NAME(InjectionLibrary)
+
+bool UUAFInjectionLibrary::IsSlotActive(const UUAFComponent* UAFComponent, FName SlotNodeName)
+{
+	if (UAFComponent && SlotNodeName != NAME_None)
+	{
+		if (const AActor* UAFActor = UAFComponent->GetOwner())
+		{
+			if (const USkeletalMeshComponent* SkeletalMeshComponent = UAFActor->FindComponentByClass<USkeletalMeshComponent>())
+			{
+				if (SkeletalMeshComponent->bEnableAnimation)
+				{
+					return false;
+				}
+				
+				if (const UAnimInstance* UsedAnimInstance = SkeletalMeshComponent->GetAnimInstance())
+				{
+					for (const FAnimMontageInstance* MontageInstance : UsedAnimInstance->MontageInstances)
+					{
+						if (MontageInstance->IsActive()
+							&& MontageInstance->Montage->IsValidSlot(SlotNodeName))
+						{
+							return true;
+						}
+					}
+				}
+			}
+		}
+	}
+	
+	return false;
+}
\ No newline at end of file
diff --git a/Engine/Plugins/Experimental/UAF/UAFAnimGraph/Source/UAFAnimGraph/Private/Injection/UAFMontageComponent.cpp b/Engine/Plugins/Experimental/UAF/UAFAnimGraph/Source/UAFAnimGraph/Private/Injection/UAFMontageComponent.cpp
--- a/Engine/Plugins/Experimental/UAF/UAFAnimGraph/Source/UAFAnimGraph/Private/Injection/UAFMontageComponent.cpp
+++ b/Engine/Plugins/Experimental/UAF/UAFAnimGraph/Source/UAFAnimGraph/Private/Injection/UAFMontageComponent.cpp
@@ -1,18 +1,18 @@
 // Copyright Epic Games, Inc. All Rights Reserved.
 
 #include "Injection/UAFMontageComponent.h"
-#include "Graph/AnimNextGraphInstance.h"
-#include "Module/AnimNextModuleInstance.h"
-#include "Module/AnimNextSkeletalMeshComponentReferenceComponent.h"
-#include "Component/AnimNextComponent.h"
-#include "Animation/AnimNotifies/AnimNotifyState.h"
-#include "Traits/NotifyDispatcher.h"
+
+#include "UAFLogging.h"
 #include "Animation/AnimInstance.h"
 #include "Animation/AnimMontage.h"
+#include "Animation/AnimNotifies/AnimNotifyState.h"
+#include "Component/AnimNextComponent.h"
 #include "Components/SkeletalMeshComponent.h"
-#include "Engine/World.h"
 #include "Engine/Level.h"
+#include "Engine/World.h"
 #include "Logging/StructuredLog.h"
+#include "Module/AnimNextModuleInstance.h"
+#include "Traits/NotifyDispatcher.h"
 
 
 #include UE_INLINE_GENERATED_CPP_BY_NAME(UAFMontageComponent)
@@ -111,6 +111,22 @@ FUAFMontageComponent::~FUAFMontageComponent()
 	}
 }
 
+bool FUAFMontageComponent::IsSlotActive(FName SlotNodeName) const
+{
+	if (SlotNodeName != NAME_None)
+	{
+		for (const FUAFMontageEvalData& EvalData : MontageEvalData)
+		{
+			if (EvalData.IsActive() && EvalData.GetMontage()->IsValidSlot(SlotNodeName))
+			{
+				return true;
+			}
+		}
+	}
+
+	return false;
+}
+
 bool FUAFMontageComponent::GetActiveNotifiesForSlot(FName SlotName, TArray<FAnimNotifyEventReference>& OutNotifies) const
 {
 	const FAnimNotifyArray* FoundEntry = ActiveSlotNotifies.Find(SlotName);
@@ -228,3 +244,27 @@ void FUAFMontageComponent::OnBindToInstance()
 		});
 }
 
+FRigUnit_IsSlotActive_Execute()
+{
+	if (SlotNodeName == NAME_None)
+	{
+		Result = false;
+		UAF_RIGUNIT_LOG(Warning, TEXT("[%s] is not a valid slot name."), *SlotNodeName.ToString());
+		return;
+	}
+	
+	const FAnimNextModuleContextData& ModuleContextData = ExecuteContext.GetContextData<FAnimNextModuleContextData>();
+	FAnimNextModuleInstance& ModuleInstance = ModuleContextData.GetModuleInstance();
+	const FUAFMontageComponent* MontageComponent = ModuleInstance.TryGetComponent<FUAFMontageComponent>();
+	if (MontageComponent)
+	{
+		Result = MontageComponent->IsSlotActive(SlotNodeName);
+	}
+	else
+	{
+		Result = false;
+		UAF_RIGUNIT_LOG(Warning, TEXT("Could not check if slot [%s] is active, montages are not setup in this system."), *SlotNodeName.ToString());
+	}
+}
+
+
diff --git a/Engine/Plugins/Experimental/UAF/UAFAnimGraph/Source/UAFAnimGraph/Public/Injection/InjectionLibrary.h b/Engine/Plugins/Experimental/UAF/UAFAnimGraph/Source/UAFAnimGraph/Public/Injection/InjectionLibrary.h
new file mode 100644
--- /dev/null
+++ b/Engine/Plugins/Experimental/UAF/UAFAnimGraph/Source/UAFAnimGraph/Public/Injection/InjectionLibrary.h
@@ -0,0 +1,22 @@
+// Copyright Epic Games, Inc. All Rights Reserved.
+
+#pragma once
+
+#include "CoreMinimal.h"
+#include "Kismet/BlueprintFunctionLibrary.h"
+#include "InjectionLibrary.generated.h"
+
+class UUAFComponent;
+
+UCLASS()
+class UUAFInjectionLibrary : public UBlueprintFunctionLibrary
+{
+	GENERATED_BODY()
+
+public:
+	// Return true if the running UAF System has an active montage in the given slot. A UAnimMontage that is playing in the slot and blending out is not determined to be "active".
+	UFUNCTION(BlueprintPure, Category = "UAF|Injection", meta=(ScriptMethod))
+	static UAFANIMGRAPH_API bool IsSlotActive(const UUAFComponent* UAFComponent, FName SlotNodeName);
+};
+
+
