UAF: Remove module context assumption in IsSlotActive RigUnit call

#Synced-CL 50551968

diff --git a/Engine/Plugins/Experimental/UAF/UAFAnimGraph/Source/UAFAnimGraph/Private/Injection/UAFMontageComponent.cpp b/Engine/Plugins/Experimental/UAF/UAFAnimGraph/Source/UAFAnimGraph/Private/Injection/UAFMontageComponent.cpp
--- a/Engine/Plugins/Experimental/UAF/UAFAnimGraph/Source/UAFAnimGraph/Private/Injection/UAFMontageComponent.cpp
+++ b/Engine/Plugins/Experimental/UAF/UAFAnimGraph/Source/UAFAnimGraph/Private/Injection/UAFMontageComponent.cpp
@@ -253,9 +253,8 @@ FRigUnit_IsSlotActive_Execute()
 		return;
 	}
 	
-	const FAnimNextModuleContextData& ModuleContextData = ExecuteContext.GetContextData<FAnimNextModuleContextData>();
-	FAnimNextModuleInstance& ModuleInstance = ModuleContextData.GetModuleInstance();
-	const FUAFMontageComponent* MontageComponent = ModuleInstance.TryGetComponent<FUAFMontageComponent>();
+	FUAFAssetInstance& Instance = ExecuteContext.GetContextData<FUAFAssetContextData>().GetInstance();
+	const FUAFMontageComponent* MontageComponent = Instance.TryGetComponent<FUAFMontageComponent>();
 	if (MontageComponent)
 	{
 		Result = MontageComponent->IsSlotActive(SlotNodeName);
