Gracefully handle notify dispatch failure when no animation sequence is playing

#Synced-CL 50652508

diff --git a/Engine/Plugins/Experimental/EvaluationNotifies/Source/Runtime/Private/AnimNext/EvaluationNotifiesTrait.cpp b/Engine/Plugins/Experimental/EvaluationNotifies/Source/Runtime/Private/AnimNext/EvaluationNotifiesTrait.cpp
--- a/Engine/Plugins/Experimental/EvaluationNotifies/Source/Runtime/Private/AnimNext/EvaluationNotifiesTrait.cpp
+++ b/Engine/Plugins/Experimental/EvaluationNotifies/Source/Runtime/Private/AnimNext/EvaluationNotifiesTrait.cpp
@@ -43,11 +43,11 @@ namespace UE::UAF
 
 		// Get current state from the stack, advance time, then get the delta state
 		FTimelineState PreAdvanceState;
-		ensure(TimelineTrait.GetState(Context, PreAdvanceState));
+		const bool bHasState = TimelineTrait.GetState(Context, PreAdvanceState);
 
 		TimelinePlayerTrait.AdvanceBy(Context, DeltaTime, bDispatchEvents);
 
-		if (!bDispatchEvents)
+		if (!bDispatchEvents || !bHasState)
 		{
 			return;
 		}
diff --git a/Engine/Plugins/Experimental/UAF/UAFAnimGraph/Source/UAFAnimGraph/Private/Traits/NotifyDispatcher.cpp b/Engine/Plugins/Experimental/UAF/UAFAnimGraph/Source/UAFAnimGraph/Private/Traits/NotifyDispatcher.cpp
--- a/Engine/Plugins/Experimental/UAF/UAFAnimGraph/Source/UAFAnimGraph/Private/Traits/NotifyDispatcher.cpp
+++ b/Engine/Plugins/Experimental/UAF/UAFAnimGraph/Source/UAFAnimGraph/Private/Traits/NotifyDispatcher.cpp
@@ -43,10 +43,11 @@ namespace UE::UAF
 
 		// Get current state from the stack, advance time, then get the delta state
 		FTimelineState PreAdvanceState;
-		ensure(TimelineTrait.GetState(Context, PreAdvanceState));
+		const bool bHasState = TimelineTrait.GetState(Context, PreAdvanceState);
+
 		TimelinePlayerTrait.AdvanceBy(Context, DeltaTime, bDispatchEvents);
 
-		if (!bDispatchEvents)
+		if (!bDispatchEvents || !bHasState)
 		{
 			return;
 		}
@@ -56,12 +57,12 @@ namespace UE::UAF
 			return;
 		}
 
-		TTraitBinding<INotifySource> NotifySourceTrait;
-		ensure(Binding.GetStackInterfaceSuper(NotifySourceTrait));
-
 		FTimelineDelta Delta;
 		if (ensure(TimelineTrait.GetDelta(Context, Delta)))
 		{
+			TTraitBinding<INotifySource> NotifySourceTrait;
+			ensure(Binding.GetStackInterfaceSuper(NotifySourceTrait));
+
 			// Query for notifies
 			TArray<FAnimNotifyEventReference> Notifies;
 			NotifySourceTrait.GetNotifies(Context, PreAdvanceState.GetPosition(), Delta.GetDeltaTime(), PreAdvanceState.IsLooping(), Notifies);
