AnimNode Starting point - AnimGraph dependencies

#Synced-CL 50467870

diff --git a/Engine/Plugins/Experimental/UAF/UAFAnimGraph/Source/UAFAnimGraph/Private/Graph/RigUnit_AnimNextRunAnimationGraph_v2.cpp b/Engine/Plugins/Experimental/UAF/UAFAnimGraph/Source/UAFAnimGraph/Private/Graph/RigUnit_AnimNextRunAnimationGraph_v2.cpp
--- a/Engine/Plugins/Experimental/UAF/UAFAnimGraph/Source/UAFAnimGraph/Private/Graph/RigUnit_AnimNextRunAnimationGraph_v2.cpp
+++ b/Engine/Plugins/Experimental/UAF/UAFAnimGraph/Source/UAFAnimGraph/Private/Graph/RigUnit_AnimNextRunAnimationGraph_v2.cpp
@@ -12,7 +12,7 @@
 #include "Injection/GraphInstanceInjectionComponent.h"
 #include "TraitCore/TraitEventList.h"
 #include "Injection/AnimNextModuleInjectionComponent.h"
-#include "TraceAnimNextGraphInstances.h"
+#include "Graph/UAFAnimGraphTrace.h"
 #include "Factory/AnimGraphFactory.h"
 #include "UAF/UAFSkeletonUserData.h"
 #include "UAFLogging.h"
diff --git a/Engine/Plugins/Experimental/UAF/UAFAnimGraph/Source/UAFAnimGraph/Private/Graph/TraceAnimNextGraphInstances.cpp b/Engine/Plugins/Experimental/UAF/UAFAnimGraph/Source/UAFAnimGraph/Private/Graph/TraceAnimNextGraphInstances.cpp
deleted file mode 100644
diff --git a/Engine/Plugins/Experimental/UAF/UAFAnimGraph/Source/UAFAnimGraph/Private/Graph/TraceAnimNextGraphInstances.h b/Engine/Plugins/Experimental/UAF/UAFAnimGraph/Source/UAFAnimGraph/Private/Graph/TraceAnimNextGraphInstances.h
deleted file mode 100644
diff --git a/Engine/Plugins/Experimental/UAF/UAFAnimGraph/Source/UAFAnimGraph/Private/Graph/UAFAnimGraphTrace.cpp b/Engine/Plugins/Experimental/UAF/UAFAnimGraph/Source/UAFAnimGraph/Private/Graph/UAFAnimGraphTrace.cpp
new file mode 100644
--- /dev/null
+++ b/Engine/Plugins/Experimental/UAF/UAFAnimGraph/Source/UAFAnimGraph/Private/Graph/UAFAnimGraphTrace.cpp
@@ -0,0 +1,114 @@
+// Copyright Epic Games, Inc. All Rights Reserved.
+
+#include "Graph/UAFAnimGraphTrace.h"
+
+#include "EvaluationVM/EvaluationProgram.h"
+#include "EvaluationVM/SerializableEvaluationProgram.h"
+#include "Graph/AnimNextGraphInstance.h"
+#include "Module/AnimNextModuleInstance.h"
+#include "ObjectTrace.h"
+#include "ObjectAsTraceIdProxyArchive.h"
+#include "Serialization/MemoryWriter.h"
+#include "TraitCore/ExecutionContext.h"
+#include "TraitCore/NodeInstance.h"
+#include "TraitInterfaces/IHierarchy.h"
+
+#if UAF_TRACE_ENABLED
+
+UE_TRACE_EVENT_BEGIN(UAFAnimGraph, EvaluationProgram)
+	UE_TRACE_EVENT_FIELD(uint64, Cycle)
+	UE_TRACE_EVENT_FIELD(double, RecordingTime)
+	UE_TRACE_EVENT_FIELD(uint64, OuterObjectId)
+	UE_TRACE_EVENT_FIELD(uint64, InstanceId)
+	UE_TRACE_EVENT_FIELD(uint8[], ProgramData)
+UE_TRACE_EVENT_END()
+
+namespace UE::UAF
+{
+	void TraceGraphInstances(const FAnimNextGraphInstance& RootGraph)
+	{
+		bool bChannelEnabled = UE_TRACE_CHANNELEXPR_IS_ENABLED(UAFChannel);
+    	if (!bChannelEnabled)
+    	{
+    		return;
+    	}
+
+		UObject* OuterObject = RootGraph.GetModuleInstance() ? RootGraph.GetModuleInstance()->GetObject() : nullptr;
+
+		FExecutionContext ExecutionContext;
+		TArray<FWeakTraitPtr, TInlineAllocator<8, TMemStackAllocator<>>> Traits;
+		TSet<uint64> TracedInstances;
+		
+		FMemStack& MemStack = FMemStack::Get();
+		FMemMark Mark(MemStack);
+		
+		ExecutionContext.BindTo(RootGraph.GetGraphRootPtr());
+
+		Traits.Push(RootGraph.GetGraphRootPtr());
+
+		TTraitBinding<IHierarchy> HierarchyTrait;
+
+		while (!Traits.IsEmpty())
+		{
+			FWeakTraitPtr Trait = Traits.Pop();
+			if (Trait.IsValid())
+			{
+				FAnimNextGraphInstance& Graph = Trait.GetNodeInstance()->GetOwner();
+				if (!TracedInstances.Contains(Graph.GetUniqueId()))
+				{
+					TRACE_UAF_VARIABLES(&Graph, Graph.GetModuleInstance() ? Graph.GetModuleInstance()->GetObject() : nullptr);
+				}
+
+				FTraitStackBinding TraitStack;
+				ExecutionContext.GetStack(Trait, TraitStack);
+				
+				if (TraitStack.GetInterface(HierarchyTrait))
+				{
+					IHierarchy::GetStackChildren(ExecutionContext, TraitStack, Traits);
+				}
+			}
+		}
+	}
+
+	void TraceEvaluationProgram(const FEvaluationProgram& Program,  const FAnimNextGraphInstance& RootGraph)
+	{
+		bool bChannelEnabled = UE_TRACE_CHANNELEXPR_IS_ENABLED(UAFChannel);
+		if (!bChannelEnabled)
+		{
+			return;
+		}
+
+		uint64 OuterObjectId = 0;
+		if (FAnimNextModuleInstance* ModuleInstance = RootGraph.GetModuleInstance())
+		{
+			TraceEvaluationProgram(Program, RootGraph, ModuleInstance->GetObject());
+		}
+	}
+		
+	void TraceEvaluationProgram(const FEvaluationProgram& Program,  const FUAFAssetInstance& RootGraph, const UObject* OuterObject)
+	{
+		uint64 InstanceId = RootGraph.GetUniqueId();
+		
+		double RecordingTime = 0;
+		if (OuterObject)
+		{
+			RecordingTime = FObjectTrace::GetWorldElapsedTime(OuterObject->GetWorld());
+		}
+
+		TArray<uint8> ArchiveData;
+		FMemoryWriter WriterArchive(ArchiveData);
+		FObjectAsTraceIdProxyArchive Archive(WriterArchive);
+		
+		static const FSerializableEvaluationProgram Defaults;
+		FSerializableEvaluationProgram SerializableProgram(Program);
+		FSerializableEvaluationProgram::StaticStruct()->SerializeItem(Archive, &SerializableProgram, &Defaults);
+		
+		UE_TRACE_LOG(UAFAnimGraph, EvaluationProgram, UAFChannel)
+				<< EvaluationProgram.Cycle(FPlatformTime::Cycles64())
+				<< EvaluationProgram.OuterObjectId(FObjectTrace::GetObjectId(OuterObject))
+				<< EvaluationProgram.RecordingTime(RecordingTime)
+				<< EvaluationProgram.InstanceId(InstanceId)
+				<< EvaluationProgram.ProgramData(ArchiveData.GetData(), ArchiveData.Num());
+	}
+}
+#endif
diff --git a/Engine/Plugins/Experimental/UAF/UAFAnimGraph/Source/UAFAnimGraph/Private/TraitInterfaces/IEvaluate.cpp b/Engine/Plugins/Experimental/UAF/UAFAnimGraph/Source/UAFAnimGraph/Private/TraitInterfaces/IEvaluate.cpp
--- a/Engine/Plugins/Experimental/UAF/UAFAnimGraph/Source/UAFAnimGraph/Private/TraitInterfaces/IEvaluate.cpp
+++ b/Engine/Plugins/Experimental/UAF/UAFAnimGraph/Source/UAFAnimGraph/Private/TraitInterfaces/IEvaluate.cpp
@@ -10,7 +10,7 @@
 #include "EvaluationVM/EvaluationVM.h"
 #include "Graph/AnimNextGraphInstance.h"
 #include "Graph/AnimNext_LODPose.h"
-#include "Graph/TraceAnimNextGraphInstances.h"
+#include "Graph/UAFAnimGraphTrace.h"
 
 #include "UAF/Attributes/EngineAttributes.h"
 #include "UAF/ValueRuntime/ValueBundle.h"
diff --git a/Engine/Plugins/Experimental/UAF/UAFAnimGraph/Source/UAFAnimGraph/Public/Graph/UAFAnimGraphTrace.h b/Engine/Plugins/Experimental/UAF/UAFAnimGraph/Source/UAFAnimGraph/Public/Graph/UAFAnimGraphTrace.h
new file mode 100644
--- /dev/null
+++ b/Engine/Plugins/Experimental/UAF/UAFAnimGraph/Source/UAFAnimGraph/Public/Graph/UAFAnimGraphTrace.h
@@ -0,0 +1,30 @@
+// Copyright Epic Games, Inc. All Rights Reserved.
+#pragma once
+
+#include "EvaluationVM/EvaluationProgram.h"
+#include "RewindDebugger/UAFTrace.h"
+
+#if UAF_TRACE_ENABLED
+
+struct FAnimNextGraphInstance;
+
+namespace UE::UAF
+{
+	struct FEvaluationProgram;
+	
+	UAFANIMGRAPH_API void TraceGraphInstances(const FAnimNextGraphInstance& RootGraph);
+	UAFANIMGRAPH_API void TraceEvaluationProgram(const FEvaluationProgram& Program, const FAnimNextGraphInstance& RootGraph);
+	UAFANIMGRAPH_API void TraceEvaluationProgram(const FEvaluationProgram& Program, const FUAFAssetInstance& AssetInstance, const UObject* OuterObject);
+}
+
+#define TRACE_UAF_GRAPHINSTANCES(RootGraph) UE::UAF::TraceGraphInstances(RootGraph);
+#define TRACE_UAF_EVALUATIONPROGRAM(Program, RootGraph) UE::UAF::TraceEvaluationProgram(Program, RootGraph);
+#define TRACE_UAF_EVALUATIONPROGRAM_WITHOWNER(Program, AssetInstance, Object) UE::UAF::TraceEvaluationProgram(Program, AssetInstance, Object);
+
+#else
+
+#define TRACE_UAF_GRAPHINSTANCES(RootGraph)
+#define TRACE_UAF_EVALUATIONPROGRAM(Program, RootGraph)
+#define TRACE_UAF_EVALUATIONPROGRAM_WITHOWNER(Program, AssetInstance, Object)
+
+#endif
