UAF: Add support for checking if a slot is active if the node is placed in a graph and remove warning log spam

#Synced-CL 50702987

diff --git a/Engine/Plugins/Experimental/UAF/UAFAnimGraph/Source/UAFAnimGraph/Private/Injection/UAFMontageComponent.cpp b/Engine/Plugins/Experimental/UAF/UAFAnimGraph/Source/UAFAnimGraph/Private/Injection/UAFMontageComponent.cpp
--- a/Engine/Plugins/Experimental/UAF/UAFAnimGraph/Source/UAFAnimGraph/Private/Injection/UAFMontageComponent.cpp
+++ b/Engine/Plugins/Experimental/UAF/UAFAnimGraph/Source/UAFAnimGraph/Private/Injection/UAFMontageComponent.cpp
@@ -246,24 +246,20 @@ void FUAFMontageComponent::OnBindToInstance()
 
 FRigUnit_IsSlotActive_Execute()
 {
+	Result = false;
 	if (SlotNodeName == NAME_None)
 	{
-		Result = false;
 		UAF_RIGUNIT_LOG(Warning, TEXT("[%s] is not a valid slot name."), *SlotNodeName.ToString());
 		return;
 	}
 	
 	FUAFAssetInstance& Instance = ExecuteContext.GetContextData<FUAFAssetContextData>().GetInstance();
-	const FUAFMontageComponent* MontageComponent = Instance.TryGetComponent<FUAFMontageComponent>();
+	FAnimNextModuleInstance* RootInstance = Instance.GetRootInstance();
+	const FUAFMontageComponent* MontageComponent = RootInstance ? RootInstance->TryGetComponent<FUAFMontageComponent>() : nullptr;
 	if (MontageComponent)
 	{
 		Result = MontageComponent->IsSlotActive(SlotNodeName);
 	}
-	else
-	{
-		Result = false;
-		UAF_RIGUNIT_LOG(Warning, TEXT("Could not check if slot [%s] is active, montages are not setup in this system."), *SlotNodeName.ToString());
-	}
 }
 
 
