UAF: Crash when playing an animation sequence that does not have a skeleton assigned

This crash happened when you remove a skeleton asset that an animation sequence is referring to, and then play the animation via a SequencePlayer node in a UAF anim graph.

#Synced-CL 50710079

diff --git a/Engine/Plugins/Experimental/UAF/UAFAnimGraph/Source/UAFAnimGraph/Private/EvaluationVM/Tasks/PushAnimSequenceKeyframe.cpp b/Engine/Plugins/Experimental/UAF/UAFAnimGraph/Source/UAFAnimGraph/Private/EvaluationVM/Tasks/PushAnimSequenceKeyframe.cpp
--- a/Engine/Plugins/Experimental/UAF/UAFAnimGraph/Source/UAFAnimGraph/Private/EvaluationVM/Tasks/PushAnimSequenceKeyframe.cpp
+++ b/Engine/Plugins/Experimental/UAF/UAFAnimGraph/Source/UAFAnimGraph/Private/EvaluationVM/Tasks/PushAnimSequenceKeyframe.cpp
@@ -36,7 +36,8 @@ void FAnimNextAnimSequenceKeyframeTask::Execute(UE::UAF::FEvaluationVM& VM) cons
 	
 	const FEvaluationTaskContext& VMContext = VM.GetActiveEvaluationContext();
 
-	if(const UAnimSequence* AnimSequencePtr = AnimSequence.Get())
+	const UAnimSequence* AnimSequencePtr = AnimSequence.Get();
+	if (AnimSequencePtr && AnimSequencePtr->GetSkeleton()) // Skip in case the anim sequence has no valid skeleton assigned.
 	{
 		const bool bExtractRootMotion = bExtractTrajectory;
 		const FAnimExtractContext ExtractionContext(SampleTime, bExtractRootMotion, DeltaTimeRecord, bLooping);
