UAF: Layered Blend
- Account for GetBlendWeight during the update of the layered blend
- Add option to always update the child branch independent of weight
- Rename the shared data to UAF

#Synced-CL 50765472

diff --git a/Engine/Plugins/Experimental/UAF/UAFAnimGraph/Config/DefaultUAFAnimGraph.ini b/Engine/Plugins/Experimental/UAF/UAFAnimGraph/Config/DefaultUAFAnimGraph.ini
--- a/Engine/Plugins/Experimental/UAF/UAFAnimGraph/Config/DefaultUAFAnimGraph.ini
+++ b/Engine/Plugins/Experimental/UAF/UAFAnimGraph/Config/DefaultUAFAnimGraph.ini
@@ -85,7 +85,7 @@
 +StructRedirects=(OldName="/Script/AnimNextAnimGraph.AnimNextBlendByBoolTraitSharedData", NewName="/Script/UAFAnimGraph.AnimNextBlendByBoolTraitSharedData")
 +StructRedirects=(OldName="/Script/AnimNextAnimGraph.AnimNextBlendInertializerCoreTraitSharedData", NewName="/Script/UAFAnimGraph.AnimNextBlendInertializerCoreTraitSharedData")
 +StructRedirects=(OldName="/Script/AnimNextAnimGraph.AnimNextBlendInertializerTraitSharedData", NewName="/Script/UAFAnimGraph.AnimNextBlendInertializerTraitSharedData")
-+StructRedirects=(OldName="/Script/AnimNextAnimGraph.AnimNextBlendLayerTraitSharedData", NewName="/Script/UAFAnimGraph.AnimNextBlendLayerTraitSharedData")
++StructRedirects=(OldName="/Script/AnimNextAnimGraph.AnimNextBlendLayerTraitSharedData", NewName="/Script/UAFAnimGraph.UAFBlendLayerTraitSharedData")
 +StructRedirects=(OldName="/Script/AnimNextAnimGraph.AnimNextBlendSmootherPerBoneCoreTraitSharedData", NewName="/Script/UAFAnimGraph.AnimNextBlendSmootherPerBoneCoreTraitSharedData")
 +StructRedirects=(OldName="/Script/AnimNextAnimGraph.AnimNextBlendProfilePerChildProviderTraitSharedData", NewName="/Script/UAFAnimGraph.AnimNextBlendProfilePerChildProviderTraitSharedData")
 +StructRedirects=(OldName="/Script/AnimNextAnimGraph.AnimNextBlendProfileProviderTraitSharedData", NewName="/Script/UAFAnimGraph.AnimNextBlendProfileProviderTraitSharedData")
@@ -207,4 +207,5 @@
 ; AnimNext -> UAF
 +ClassRedirects=(OldName="/Script/UAFAnimGraph.AnimNextAnimationGraph",NewName="/Script/UAFAnimGraph.UAFAnimGraph")
 +ClassRedirects=(OldName="/Script/UAFAnimGraphUncookedOnly.AnimNextAnimationGraph_EditorData",NewName="/Script/UAFAnimGraphUncookedOnly.UAFAnimGraph_EditorData")
-+ClassRedirects=(OldName="/Script/UAFAnimGraphUncookedOnly.AnimNextAnimationGraphFactory",NewName="/Script/UAFAnimGraphUncookedOnly.UAFAnimGraphFactory")
\ No newline at end of file
++ClassRedirects=(OldName="/Script/UAFAnimGraphUncookedOnly.AnimNextAnimationGraphFactory",NewName="/Script/UAFAnimGraphUncookedOnly.UAFAnimGraphFactory")
++StructRedirects=(OldName="/Script/UAFAnimGraph.AnimNextBlendLayerTraitSharedData",NewName="/Script/UAFAnimGraph.UAFBlendLayerTraitSharedData")
diff --git a/Engine/Plugins/Experimental/UAF/UAFAnimGraph/Source/UAFAnimGraph/Private/Traits/BlendLayer.cpp b/Engine/Plugins/Experimental/UAF/UAFAnimGraph/Source/UAFAnimGraph/Private/Traits/BlendLayer.cpp
--- a/Engine/Plugins/Experimental/UAF/UAFAnimGraph/Source/UAFAnimGraph/Private/Traits/BlendLayer.cpp
+++ b/Engine/Plugins/Experimental/UAF/UAFAnimGraph/Source/UAFAnimGraph/Private/Traits/BlendLayer.cpp
@@ -180,8 +180,9 @@ namespace UE::UAF
 		{
 			InstanceData->bWasChildBaseRelevant = true;
 		}
-
-		if (FAnimWeight::IsRelevant(SharedData->BlendWeight) && SharedData->BlendProfile)
+		
+		const float ChildWeight = ContinuousBlendTrait.GetBlendWeight(Context, 1);
+		if ((FAnimWeight::IsRelevant(ChildWeight) || SharedData->GetbAlwaysUpdateChildBlend(Binding)) && SharedData->BlendProfile)
 		{
 			if (!InstanceData->ChildBlend.IsValid())
 			{
diff --git a/Engine/Plugins/Experimental/UAF/UAFAnimGraph/Source/UAFAnimGraph/Private/Traits/BlendLayer.h b/Engine/Plugins/Experimental/UAF/UAFAnimGraph/Source/UAFAnimGraph/Private/Traits/BlendLayer.h
--- a/Engine/Plugins/Experimental/UAF/UAFAnimGraph/Source/UAFAnimGraph/Private/Traits/BlendLayer.h
+++ b/Engine/Plugins/Experimental/UAF/UAFAnimGraph/Source/UAFAnimGraph/Private/Traits/BlendLayer.h
@@ -24,7 +24,7 @@ namespace UE::UAF
 		DECLARE_ANIM_TRAIT(FBlendLayerTrait, FBaseTrait)
 
 
-		using FSharedData = FAnimNextBlendLayerTraitSharedData;
+		using FSharedData = FUAFBlendLayerTraitSharedData;
 
 		struct FInstanceData : FTrait::FInstanceData
 		{
diff --git a/Engine/Plugins/Experimental/UAF/UAFAnimGraph/Source/UAFAnimGraph/Public/Traits/BlendLayerTraitData.h b/Engine/Plugins/Experimental/UAF/UAFAnimGraph/Source/UAFAnimGraph/Public/Traits/BlendLayerTraitData.h
--- a/Engine/Plugins/Experimental/UAF/UAFAnimGraph/Source/UAFAnimGraph/Public/Traits/BlendLayerTraitData.h
+++ b/Engine/Plugins/Experimental/UAF/UAFAnimGraph/Source/UAFAnimGraph/Public/Traits/BlendLayerTraitData.h
@@ -12,7 +12,7 @@ class UHierarchyTable;
 
 /** A trait that can blend a layer into an input. */
 USTRUCT(meta = (DisplayName = "Blend Layer", ShowTooltip=true))
-struct FAnimNextBlendLayerTraitSharedData : public FAnimNextTraitSharedData
+struct FUAFBlendLayerTraitSharedData : public FAnimNextTraitSharedData
 {
 	GENERATED_BODY()
 
@@ -34,11 +34,16 @@ struct FAnimNextBlendLayerTraitSharedData : public FAnimNextTraitSharedData
 	// Whether or not to blend in mesh space or local space
 	UPROPERTY(EditAnywhere, Category = "Default")
 	bool bBlendInMeshSpace = false;
+	
+	// If true, this will still update the child branch even if the blend weight is 0
+	UPROPERTY(EditAnywhere, Category = "Default")
+	bool bAlwaysUpdateChildBlend = false;
 
 	#define TRAIT_LATENT_PROPERTIES_ENUMERATOR(GeneratorMacro) \
 	GeneratorMacro(BlendWeight) \
 	GeneratorMacro(bBlendInMeshSpace) \
+	GeneratorMacro(bAlwaysUpdateChildBlend) \
 
-	GENERATE_TRAIT_LATENT_PROPERTIES(FAnimNextBlendLayerTraitSharedData, TRAIT_LATENT_PROPERTIES_ENUMERATOR)
+	GENERATE_TRAIT_LATENT_PROPERTIES(FUAFBlendLayerTraitSharedData, TRAIT_LATENT_PROPERTIES_ENUMERATOR)
 	#undef TRAIT_LATENT_PROPERTIES_ENUMERATOR
 };
\ No newline at end of file
diff --git a/Engine/Plugins/Experimental/UAF/UAFLayering/Source/UAFLayeringUncookedOnly/Private/Layers/UAFLayerDefaultBlendProvider.cpp b/Engine/Plugins/Experimental/UAF/UAFLayering/Source/UAFLayeringUncookedOnly/Private/Layers/UAFLayerDefaultBlendProvider.cpp
--- a/Engine/Plugins/Experimental/UAF/UAFLayering/Source/UAFLayeringUncookedOnly/Private/Layers/UAFLayerDefaultBlendProvider.cpp
+++ b/Engine/Plugins/Experimental/UAF/UAFLayering/Source/UAFLayeringUncookedOnly/Private/Layers/UAFLayerDefaultBlendProvider.cpp
@@ -148,16 +148,16 @@ URigVMPin* FUAFDefaultBlendProvider::CreateBlendGraphTrait(UE::UAF::Layering::FL
 		if (bUsesPerBoneData)
 		{
 			// Setup default trait values 
-			FAnimNextBlendLayerTraitSharedData DefaultTraitData; 
+			FUAFBlendLayerTraitSharedData DefaultTraitData; 
 			DefaultTraitData.BlendWeight = LayerWeight;
 			DefaultTraitData.BlendProfile = BlendMask;
 			
 			// Create trait 
-			BlendBaseTraitName = LayerCreationContext.GraphController->AddTraitStruct(BlendTraitStackNode, TInstancedStruct<FAnimNextBlendLayerTraitSharedData>::Make(DefaultTraitData), INDEX_NONE, false, false);
+			BlendBaseTraitName = LayerCreationContext.GraphController->AddTraitStruct(BlendTraitStackNode, TInstancedStruct<FUAFBlendLayerTraitSharedData>::Make(DefaultTraitData), INDEX_NONE, false, false);
 
 			// Set Input Pins so links to this trait stack can be generated 
-			ChildAInputPin = BlendTraitStackNode->FindTrait(BlendBaseTraitName, GET_MEMBER_NAME_STRING_CHECKED(FAnimNextBlendLayerTraitSharedData, ChildBase));
-			ChildBInputPin = BlendTraitStackNode->FindTrait(BlendBaseTraitName, GET_MEMBER_NAME_STRING_CHECKED(FAnimNextBlendLayerTraitSharedData, ChildBlend));
+			ChildAInputPin = BlendTraitStackNode->FindTrait(BlendBaseTraitName, GET_MEMBER_NAME_STRING_CHECKED(FUAFBlendLayerTraitSharedData, ChildBase));
+			ChildBInputPin = BlendTraitStackNode->FindTrait(BlendBaseTraitName, GET_MEMBER_NAME_STRING_CHECKED(FUAFBlendLayerTraitSharedData, ChildBlend));
 		}
 		else
 		{
