Remove unnecessary warning that destruction isn't possible yet when references remain
Remove unnecessary warning when attempting to raise a consumed trait event

#Synced-CL 50772073

diff --git a/Engine/Plugins/Experimental/UAF/UAFAnimGraph/Source/UAFAnimGraph/Private/TraitCore/ExecutionContext.cpp b/Engine/Plugins/Experimental/UAF/UAFAnimGraph/Source/UAFAnimGraph/Private/TraitCore/ExecutionContext.cpp
--- a/Engine/Plugins/Experimental/UAF/UAFAnimGraph/Source/UAFAnimGraph/Private/TraitCore/ExecutionContext.cpp
+++ b/Engine/Plugins/Experimental/UAF/UAFAnimGraph/Source/UAFAnimGraph/Private/TraitCore/ExecutionContext.cpp
@@ -329,8 +329,7 @@ namespace UE::UAF
 
 		if (NodeInstance->RemoveReference())
 		{
-			UE_LOG(LogAnimation, Warning, TEXT("FExecutionContext::ReleaseNodeInstance: Could not release node instance - Node Instance still has references"));
-			return;	// Node instance still has references, we can't release it
+			return;	// Node instance still has references, we can't release it yet
 		}
 
 		const FNodeDescription& NodeDesc = GetNodeDescription(*NodeInstance);
diff --git a/Engine/Plugins/Experimental/UAF/UAFAnimGraph/Source/UAFAnimGraph/Private/TraitCore/TraitEventRaising.cpp b/Engine/Plugins/Experimental/UAF/UAFAnimGraph/Source/UAFAnimGraph/Private/TraitCore/TraitEventRaising.cpp
--- a/Engine/Plugins/Experimental/UAF/UAFAnimGraph/Source/UAFAnimGraph/Private/TraitCore/TraitEventRaising.cpp
+++ b/Engine/Plugins/Experimental/UAF/UAFAnimGraph/Source/UAFAnimGraph/Private/TraitCore/TraitEventRaising.cpp
@@ -12,9 +12,14 @@ namespace UE::UAF
 {
 	void RaiseTraitEvent(FExecutionContext& Context, const FTraitStackBinding& StackBinding, FAnimNextTraitEvent& Event)
 	{
-		if (!StackBinding.IsValid() || !Event.IsValid())
+		if (!Event.IsValid())
 		{
-			UE_LOG(LogAnimation, Warning, TEXT("RaiseTraitEvent: Could not raise trait event - Either the stack binding or trait event are invalid."));
+			return;	// Consumed events are invalid and do not propagate further
+		}
+
+		if (!StackBinding.IsValid())
+		{
+			UE_LOG(LogAnimation, Warning, TEXT("RaiseTraitEvent: Could not raise trait event - The stack binding is invalid."));
 			return;
 		}
 
