UAF: Added trait deprecation and variable drag-drop support to templates

Trait deprecation allows upgrade workflows similar to node upgrade in RigVM
Drag-drop variable support allows custom template nodesa to be spawned from drag-dropping variables onto animation graphs

#Synced-CL 50766328

diff --git a/Engine/Plugins/Experimental/UAF/UAF/Source/UAF/Public/Param/ParamType.h b/Engine/Plugins/Experimental/UAF/UAF/Source/UAF/Public/Param/ParamType.h
--- a/Engine/Plugins/Experimental/UAF/UAF/Source/UAF/Public/Param/ParamType.h
+++ b/Engine/Plugins/Experimental/UAF/UAF/Source/UAF/Public/Param/ParamType.h
@@ -50,11 +50,9 @@ namespace UE::UAF::Private
 }
 
 /**
- * Representation of a parameter's type. Serializable, but fairly heavyweight to pass around and compare.
- * Faster comparisons and other operations can be performed on UE::UAF::FParamTypeHandle, but they cannot be
- * serialized as they are not stable across runs.
+ * Representation of a UAF variable's type.
  */
-USTRUCT()
+USTRUCT(BlueprintType, DisplayName = "UAF Variable Type")
 struct FAnimNextParamType
 {
 public:
diff --git a/Engine/Plugins/Experimental/UAF/UAF/Source/UAFEditor/Internal/Common/UAFVariableDragDropContext.h b/Engine/Plugins/Experimental/UAF/UAF/Source/UAFEditor/Internal/Common/UAFVariableDragDropContext.h
new file mode 100644
--- /dev/null
+++ b/Engine/Plugins/Experimental/UAF/UAF/Source/UAFEditor/Internal/Common/UAFVariableDragDropContext.h
@@ -0,0 +1,23 @@
+﻿// Copyright Epic Games, Inc. All Rights Reserved.
+
+#pragma once
+
+#include "UAFVariableDragDropContext.generated.h"
+
+class UEdGraph;
+
+UCLASS()
+class UUAFVariableDragDropContext : public UObject
+{
+	GENERATED_BODY()
+
+public:
+	UPROPERTY()
+	TObjectPtr<UEdGraph> ParentGraph;
+
+	FVector2f Location;
+	FName Name;
+	FString TypeObjectPath;
+	FString SourceObjectPath;
+	FString TypeName;
+};
\ No newline at end of file
diff --git a/Engine/Plugins/Experimental/UAF/UAF/Source/UAFEditor/Private/Common/GraphEditorSchemaActions.cpp b/Engine/Plugins/Experimental/UAF/UAF/Source/UAFEditor/Private/Common/GraphEditorSchemaActions.cpp
--- a/Engine/Plugins/Experimental/UAF/UAF/Source/UAFEditor/Private/Common/GraphEditorSchemaActions.cpp
+++ b/Engine/Plugins/Experimental/UAF/UAF/Source/UAFEditor/Private/Common/GraphEditorSchemaActions.cpp
@@ -18,6 +18,8 @@
 #include "AnimNextRigVMAssetEditorData.h"
 #include "UAFCompilationScope.h"
 #include "ScopedTransaction.h"
+#include "ToolMenus.h"
+#include "Common/UAFVariableDragDropContext.h"
 #include "Variables/AnimNextSharedVariables.h"
 #include "UObject/Class.h"
 
@@ -335,43 +337,63 @@ UEdGraphNode* FAnimNextSchemaAction_Variable::PerformAction(UEdGraph* ParentGrap
 
 	if (VariableAccessorChoice == EVariableAccessorChoice::Deferred)
 	{
-		FMenuBuilder MenuBuilder(true, nullptr);
-		const FText SectionText = FText::FromString(FString::Printf(TEXT("Variable %s"), *Name.ToString()));
+		static const FName MenuName = "UAFVariableDragDropMenu";
+		if(!UToolMenus::Get()->IsMenuRegistered(MenuName))
+		{
+			UToolMenu* Menu = UToolMenus::Get()->RegisterMenu(MenuName);
+			Menu->AddDynamicSection("VariableGetterSetter", FNewToolMenuDelegate::CreateLambda([AddNodeLambda](UToolMenu* InMenu)
+			{
+				const UUAFVariableDragDropContext* Context = InMenu->FindContext<UUAFVariableDragDropContext>();
+				if (Context == nullptr)
+				{
+					return;
+				}
 
-		MenuBuilder.BeginSection("VariableDropped", SectionText);
+				FToolMenuSection& Section = InMenu->AddSection("VariableGetterSetter", FText::Format(LOCTEXT("SectionFormat", "Variable {0}"), FText::FromName(Context->Name)));
 
-		MenuBuilder.AddMenuEntry(
-			FText::FromString(FString::Printf(TEXT("Get %s"), *Name.ToString())),
-			FText::FromString(FString::Printf(TEXT("Adds a getter node for variable %s"), *Name.ToString())),
-			FSlateIcon(),
-			FUIAction(
-				FExecuteAction::CreateLambda([AddNodeLambda, ParentGraph, Location, Name = Name, SourceObjectPath = SourceObjectPath, TypeName = TypeName, TypeObjectPath = TypeObjectPath]
-				{
-					AddNodeLambda(ParentGraph, Name, SourceObjectPath, TypeName, TypeObjectPath, Location, true);
-				}),
-				FCanExecuteAction()
-			)
-		);
+				Section.AddMenuEntry(
+					"GetVariable",
+					FText::Format(LOCTEXT("GetterFormat", "Get {0}"), FText::FromName(Context->Name)),
+					FText::Format(LOCTEXT("GetterTooltipFormat", "Adds a getter node for variable {0}"), FText::FromName(Context->Name)),
+					FSlateIcon(),
+					FUIAction(
+						FExecuteAction::CreateLambda([AddNodeLambda, Context]
+						{
+							AddNodeLambda(Context->ParentGraph, Context->Name, Context->SourceObjectPath, Context->TypeName, Context->TypeObjectPath, Context->Location, true);
+						}),
+						FCanExecuteAction()
+					));
+
+				Section.AddMenuEntry(
+					"SetVariable",
+					FText::Format(LOCTEXT("SetterFormat", "Set {0}"), FText::FromName(Context->Name)),
+					FText::Format(LOCTEXT("SetterTooltipFormat", "Adds a setter node for variable {0}"), FText::FromName(Context->Name)),
+					FSlateIcon(),
+					FUIAction(
+						FExecuteAction::CreateLambda([AddNodeLambda, Context]
+						{
+							AddNodeLambda(Context->ParentGraph, Context->Name, Context->SourceObjectPath, Context->TypeName, Context->TypeObjectPath, Context->Location, false);
+						}),
+						FCanExecuteAction()
+					));
+			}));
+		}
 
-		MenuBuilder.AddMenuEntry(
-			FText::FromString(FString::Printf(TEXT("Set %s"), *Name.ToString())),
-			FText::FromString(FString::Printf(TEXT("Adds a setter node for variable %s"), *Name.ToString())),
-			FSlateIcon(),
-			FUIAction(
-				FExecuteAction::CreateLambda([AddNodeLambda, ParentGraph, Location, Name = Name, SourceObjectPath = SourceObjectPath, TypeName = TypeName, ObjectPath = TypeObjectPath]
-				{
-					AddNodeLambda(ParentGraph, Name, SourceObjectPath, TypeName, ObjectPath, Location, false);
-				}),
-				FCanExecuteAction()
-			)
-		);
+		UUAFVariableDragDropContext* Context = NewObject<UUAFVariableDragDropContext>();
+		Context->ParentGraph = ParentGraph;
+		Context->Name = Name;
+		Context->Location = Location;
+		Context->SourceObjectPath = SourceObjectPath;
+		Context->TypeName = TypeName;
+		Context->TypeObjectPath = TypeObjectPath;
 
-		MenuBuilder.EndSection();
+		FToolMenuContext ToolMenuContext;
+		ToolMenuContext.AddObject(Context);
 
 		FSlateApplication::Get().PushMenu(
 			FSlateApplication::Get().GetInteractiveTopLevelWindows()[0],
 			FWidgetPath(),
-			MenuBuilder.MakeWidget(),
+			UToolMenus::Get()->GenerateWidget(MenuName, ToolMenuContext),
 			FSlateApplication::Get().GetCursorPos(),
 			FPopupTransitionEffect(FPopupTransitionEffect::ContextMenu)
 		);
diff --git a/Engine/Plugins/Experimental/UAF/UAFAnimGraph/Source/UAFAnimGraph/Internal/Graph/RigDecorator_AnimNextCppTrait.h b/Engine/Plugins/Experimental/UAF/UAFAnimGraph/Source/UAFAnimGraph/Internal/Graph/RigDecorator_AnimNextCppTrait.h
--- a/Engine/Plugins/Experimental/UAF/UAFAnimGraph/Source/UAFAnimGraph/Internal/Graph/RigDecorator_AnimNextCppTrait.h
+++ b/Engine/Plugins/Experimental/UAF/UAFAnimGraph/Source/UAFAnimGraph/Internal/Graph/RigDecorator_AnimNextCppTrait.h
@@ -25,12 +25,14 @@ struct FRigDecorator_AnimNextCppDecorator : public FRigVMTrait
 
 #if WITH_EDITOR
 	UE_API virtual void GetProgrammaticPins(URigVMController* InController, int32 InParentPinIndex, const URigVMPin* InTraitPin, const FString& InDefaultValue, struct FRigVMPinInfoArray& OutPinArray) const override;
-	
+
 	virtual UScriptStruct* GetTraitSharedDataStruct() const override
 	{
 		return DecoratorSharedDataStruct;
 	}
 
+	UE_API virtual FRigVMStructUpgradeInfo GetUpgradeInfo() const override;
+
 	UE_API const UE::UAF::FTrait* GetTrait() const;
 #endif
 };
diff --git a/Engine/Plugins/Experimental/UAF/UAFAnimGraph/Source/UAFAnimGraph/Private/Graph/RigDecorator_AnimNextCppTrait.cpp b/Engine/Plugins/Experimental/UAF/UAFAnimGraph/Source/UAFAnimGraph/Private/Graph/RigDecorator_AnimNextCppTrait.cpp
--- a/Engine/Plugins/Experimental/UAF/UAFAnimGraph/Source/UAFAnimGraph/Private/Graph/RigDecorator_AnimNextCppTrait.cpp
+++ b/Engine/Plugins/Experimental/UAF/UAFAnimGraph/Source/UAFAnimGraph/Private/Graph/RigDecorator_AnimNextCppTrait.cpp
@@ -106,8 +106,16 @@ void FRigDecorator_AnimNextCppDecorator::GetProgrammaticPins(URigVMController* I
 	}
 }
 
+FRigVMStructUpgradeInfo FRigDecorator_AnimNextCppDecorator::GetUpgradeInfo() const
+{
+	const UE::UAF::FTrait* Trait = GetTrait();
+	check(Trait != nullptr);
+	return Trait->GetUpgradeInfo();
+}
+
 const UE::UAF::FTrait* FRigDecorator_AnimNextCppDecorator::GetTrait() const
 {
 	return UE::UAF::FTraitRegistry::Get().Find(DecoratorSharedDataStruct);
 }
+
 #endif
diff --git a/Engine/Plugins/Experimental/UAF/UAFAnimGraph/Source/UAFAnimGraph/Private/TraitCore/Trait.cpp b/Engine/Plugins/Experimental/UAF/UAFAnimGraph/Source/UAFAnimGraph/Private/TraitCore/Trait.cpp
--- a/Engine/Plugins/Experimental/UAF/UAFAnimGraph/Source/UAFAnimGraph/Private/TraitCore/Trait.cpp
+++ b/Engine/Plugins/Experimental/UAF/UAFAnimGraph/Source/UAFAnimGraph/Private/TraitCore/Trait.cpp
@@ -284,6 +284,11 @@ namespace UE::UAF
 		return NumHandlesAdded;
 	}
 
+	FRigVMStructUpgradeInfo FTrait::GetUpgradeInfo() const
+	{
+		return FRigVMStructUpgradeInfo();
+	}
+
 	FTraitStaticInitHook::FTraitStaticInitHook(TraitConstructorFunc InTraitConstructor)
 		: TraitConstructor(InTraitConstructor)
 	{
diff --git a/Engine/Plugins/Experimental/UAF/UAFAnimGraph/Source/UAFAnimGraph/Public/TraitCore/Trait.h b/Engine/Plugins/Experimental/UAF/UAFAnimGraph/Source/UAFAnimGraph/Public/TraitCore/Trait.h
--- a/Engine/Plugins/Experimental/UAF/UAFAnimGraph/Source/UAFAnimGraph/Public/TraitCore/Trait.h
+++ b/Engine/Plugins/Experimental/UAF/UAFAnimGraph/Source/UAFAnimGraph/Public/TraitCore/Trait.h
@@ -3,6 +3,7 @@
 #pragma once
 
 #include "CoreMinimal.h"
+#include "RigVMCore/RigVMStructUpgradeInfo.h"
 
 #include "TraitCore/TraitEvent.h"
 #include "TraitCore/TraitHandle.h"			// Derived types are likely to refer to other traits as children
@@ -482,7 +483,10 @@ namespace UE::UAF
 		// WARNING: Latent property handles are in the order defined by the enumerator macro within the shared data,
 		// @see FAnimNextTraitSharedData::GetLatentPropertyIndex
 		virtual int32 GetLatentPropertyIndex(const FAnimNextTraitSharedData& InSharedData, FName PropertyName) const { return INDEX_NONE; }
-		
+
+		// Get any upgrade info needed (e.g. to resolve deprecation)
+		UE_API virtual FRigVMStructUpgradeInfo GetUpgradeInfo() const;
+
 	protected:
 		// Implements GetLatentPropertyMemoryLayout() by allowing a map to be provided for caching purposes to speed up repeated queries
 		UE_API FTraitLatentPropertyMemoryLayout GetLatentPropertyMemoryLayoutImpl(
diff --git a/Engine/Plugins/Experimental/UAF/UAFAnimGraph/Source/UAFAnimGraphEditor/Private/Graph/AnimNextAnimationGraphMenuExtensions.cpp b/Engine/Plugins/Experimental/UAF/UAFAnimGraph/Source/UAFAnimGraphEditor/Private/Graph/AnimNextAnimationGraphMenuExtensions.cpp
--- a/Engine/Plugins/Experimental/UAF/UAFAnimGraph/Source/UAFAnimGraphEditor/Private/Graph/AnimNextAnimationGraphMenuExtensions.cpp
+++ b/Engine/Plugins/Experimental/UAF/UAFAnimGraph/Source/UAFAnimGraphEditor/Private/Graph/AnimNextAnimationGraphMenuExtensions.cpp
@@ -27,9 +27,11 @@
 #include "Editor.h"
 #include "ScopedTransaction.h"
 #include "AssetRegistry/AssetRegistryModule.h"
+#include "Common/UAFVariableDragDropContext.h"
 #include "Kismet2/KismetEditorUtilities.h"
 #include "RigVMModel/RigVMClient.h"
 #include "Subsystems/AssetEditorSubsystem.h"
+#include "Templates/GraphNodeTemplateRegistry.h"
 #include "Templates/UAFGraphNodeTemplate.h"
 #include "Variables/AnimNextSharedVariables.h"
 
@@ -46,13 +48,13 @@ namespace UE::UAF::Editor
 void FAnimationGraphMenuExtensions::RegisterMenus()
 {
 	FToolMenuOwnerScoped OwnerScoped(TEXT("FAnimNextAnimationGraphItemDetails"));
-	UToolMenu* Menu = UToolMenus::Get()->ExtendMenu("GraphEditor.GraphNodeContextMenu.AnimNextEdGraphNode");
-	if (Menu == nullptr)
+	UToolMenu* NodeMenu = UToolMenus::Get()->ExtendMenu("GraphEditor.GraphNodeContextMenu.AnimNextEdGraphNode");
+	if (NodeMenu == nullptr)
 	{
 		return;
 	}
 
-	Menu->AddDynamicSection(TEXT("AnimNextEdGraphNode"), FNewToolMenuDelegate::CreateLambda([](UToolMenu* InMenu)
+	NodeMenu->AddDynamicSection(TEXT("AnimNextEdGraphNode"), FNewToolMenuDelegate::CreateLambda([](UToolMenu* InMenu)
 	{
 		UGraphNodeContextMenuContext* Context = InMenu->FindContext<UGraphNodeContextMenuContext>();
 		if(Context == nullptr)
@@ -438,6 +440,91 @@ void FAnimationGraphMenuExtensions::RegisterMenus()
 			}
 		}
 	}));
+
+	if (UToolMenu* DragDropMenu = UToolMenus::Get()->ExtendMenu("UAFVariableDragDropMenu"))
+	{
+		DragDropMenu->AddDynamicSection("AddTemplateNode", FNewToolMenuDelegate::CreateLambda([](UToolMenu* InMenu)
+		{
+			const UUAFVariableDragDropContext* Context = InMenu->FindContext<UUAFVariableDragDropContext>();
+			if (Context == nullptr)
+			{
+				return;
+			}
+
+			FAnimNextParamType VariableType = FAnimNextParamType::FromString(Context->TypeName);
+			if (!VariableType.IsValid())
+			{
+				return;
+			}
+
+			TConstArrayView<FGraphNodeTemplateInfo> Handlers = FGraphNodeTemplateRegistry::GetDragDropHandlersForVariable(VariableType);
+			for (const FGraphNodeTemplateInfo& Handler : Handlers)
+			{
+				FSoftObjectPath SoftClassPath(Handler.ClassPath);
+				UClass* Class = Cast<UClass>(SoftClassPath.TryLoad());
+				if (Class == nullptr)
+				{
+					continue;
+				}
+
+				UUAFGraphNodeTemplate* Template = Class->GetDefaultObject<UUAFGraphNodeTemplate>();
+
+				FEditorScriptExecutionGuard AllowScripts;
+				FText MenuDesc = FText::Format(LOCTEXT("TemplateLabelFormat", "Add {0}"), Template->GetMenuDescription());
+				FText Tooltip = Template->GetTooltipText();
+
+				InMenu->AddMenuEntry("AddTemplateNode",
+					FToolMenuEntry::InitMenuEntry(
+					"",
+					MenuDesc,
+					Tooltip,
+					FSlateIcon(),
+					FUIAction(
+						FExecuteAction::CreateLambda([Template, Context, VariableType]()
+						{
+							URigVMEdGraph* EdGraph = Cast<URigVMEdGraph>(Context->ParentGraph);
+							if (EdGraph == nullptr)
+							{
+								return;
+							}
+
+							UAnimNextController* Controller = Cast<UAnimNextController>(EdGraph->GetController());
+							if (Controller == nullptr)
+							{
+								return;
+							}
+
+							UObject* SourceObject = nullptr;
+							if (!Context->SourceObjectPath.IsEmpty())
+							{
+								SourceObject = RigVMTypeUtils::FindObjectFromCPPTypeObjectPath<UObject>(Context->SourceObjectPath);
+							}
+
+							if (SourceObject == nullptr)
+							{
+								return;
+							}
+
+							Controller->OpenUndoBracket(TEXT("Drop variables on graph"));
+
+							URigVMUnitNode* NewNode = Template->CreateNewNode(Controller, FVector2D(Context->Location));
+							if (NewNode == nullptr)
+							{
+								Controller->CancelUndoBracket();
+								return;
+							}
+
+							{
+								FEditorScriptExecutionGuard AllowScripts;
+								Template->HandleVariableDropped(Controller, NewNode, FAnimNextVariableReference::FromName(Context->Name, SourceObject));
+							}
+
+							Controller->CloseUndoBracket();
+						}))
+					));
+				}
+		}));
+	}
 }
 
 void FAnimationGraphMenuExtensions::UnregisterMenus()
diff --git a/Engine/Plugins/Experimental/UAF/UAFAnimGraph/Source/UAFAnimGraphEditor/Private/TraitStackEditor.cpp b/Engine/Plugins/Experimental/UAF/UAFAnimGraph/Source/UAFAnimGraphEditor/Private/TraitStackEditor.cpp
--- a/Engine/Plugins/Experimental/UAF/UAFAnimGraph/Source/UAFAnimGraphEditor/Private/TraitStackEditor.cpp
+++ b/Engine/Plugins/Experimental/UAF/UAFAnimGraph/Source/UAFAnimGraphEditor/Private/TraitStackEditor.cpp
@@ -560,226 +560,313 @@ namespace Private
 				.AutoWidth()
 				.HAlign(HAlign_Right)
 				[
-					SNew(SSimpleComboButton)
-					.HasDownArrow(true)
-					.OnGetMenuContent_Lambda([this]()
-					{
-						constexpr bool bShouldCloseWindowAfterMenuSelection = true;
-						FMenuBuilder MenuBuilder(bShouldCloseWindowAfterMenuSelection, nullptr);
-
-						MenuBuilder.AddSubMenu(
-							LOCTEXT("AddLabel", "Add..."),
-							LOCTEXT("AddTooltip", "Add a new trait after this one"),
-							FNewMenuDelegate::CreateLambda([this](FMenuBuilder& InMenuBuilder)
+					SNew(SHorizontalBox)
+					+SHorizontalBox::Slot()
+					.AutoWidth()
+					.VAlign(VAlign_Center)
+					.Padding(3.0f, 0.0f)
+					[
+						SNew(SImage)
+						.Image(FAppStyle::Get().GetBrush("Icons.Warning"))
+						.ColorAndOpacity(FAppStyle::Get().GetSlateColor("Colors.AccentYellow"))
+						.ToolTipText(LOCTEXT("DeprecatedWarning", "This trait is deprecated. Please remove, replace or upgrade it to remove this warning."))
+						.Visibility_Lambda([this]()
+						{
+							TSharedPtr<IPropertyHandle> PropertyHandle = WeakPropertyHandle.Pin();
+							if (!PropertyHandle.IsValid())
 							{
-								FStructViewerModule& StructViewerModule = FModuleManager::LoadModuleChecked<FStructViewerModule>("StructViewer");
-
-								FStructViewerInitializationOptions Options;
-								Options.Mode = EStructViewerMode::StructPicker;
-								Options.NameTypeToDisplay = EStructViewerNameTypeToDisplay::DisplayName;
-								Options.StructFilter = MakeShared<FStructFilter>(Nodes, ETraitMode::Additive);
-
-								InMenuBuilder.AddWidget(
-									SNew(SBox)
-									.WidthOverride(300.0f)
-									.HeightOverride(300.0f)
-									[
-										StructViewerModule.CreateStructViewer(Options, FOnStructPicked::CreateLambda([this](const UScriptStruct* InStruct)
-										{
-											FSlateApplication::Get().DismissAllMenus();
+								return EVisibility::Collapsed;
+							}
 
-											UAnimNextController* Controller = Cast<UAnimNextController>(WeakController.Get());
-											if (Controller == nullptr)
-											{
-												return;
-											}
+							TSharedPtr<IPropertyHandleStruct> StructHandle = PropertyHandle->AsStruct();
+							check(StructHandle.IsValid());
+
+							TSharedPtr<FStructOnScope> StructData = StructHandle->GetStructData();
+							return StructData.IsValid() && StructData->GetStruct() && StructData->GetStruct()->HasMetaData(TEXT("Deprecated")) ? EVisibility::Visible : EVisibility::Collapsed;
+						})
+					]
+					+SHorizontalBox::Slot()
+					.AutoWidth()
+					.VAlign(VAlign_Center)
+					[
+						SNew(SSimpleComboButton)
+						.HasDownArrow(true)
+						.OnGetMenuContent_Lambda([this]()
+						{
+							constexpr bool bShouldCloseWindowAfterMenuSelection = true;
+							FMenuBuilder MenuBuilder(bShouldCloseWindowAfterMenuSelection, nullptr);
 
-											TSharedPtr<IPropertyHandle> PropertyHandle = WeakPropertyHandle.Pin();
-											if (!PropertyHandle.IsValid())
+							MenuBuilder.AddSubMenu(
+								LOCTEXT("AddLabel", "Add..."),
+								LOCTEXT("AddTooltip", "Add a new trait after this one"),
+								FNewMenuDelegate::CreateLambda([this](FMenuBuilder& InMenuBuilder)
+								{
+									FStructViewerModule& StructViewerModule = FModuleManager::LoadModuleChecked<FStructViewerModule>("StructViewer");
+
+									FStructViewerInitializationOptions Options;
+									Options.Mode = EStructViewerMode::StructPicker;
+									Options.NameTypeToDisplay = EStructViewerNameTypeToDisplay::DisplayName;
+									Options.StructFilter = MakeShared<FStructFilter>(Nodes, ETraitMode::Additive);
+
+									InMenuBuilder.AddWidget(
+										SNew(SBox)
+										.WidthOverride(300.0f)
+										.HeightOverride(300.0f)
+										[
+											StructViewerModule.CreateStructViewer(Options, FOnStructPicked::CreateLambda([this](const UScriptStruct* InStruct)
 											{
-												return;
-											}
+												FSlateApplication::Get().DismissAllMenus();
+
+												UAnimNextController* Controller = Cast<UAnimNextController>(WeakController.Get());
+												if (Controller == nullptr)
+												{
+													return;
+												}
 
-											TSharedPtr<IPropertyHandleStruct> StructHandle = PropertyHandle->AsStruct();
-											check(StructHandle.IsValid());
+												TSharedPtr<IPropertyHandle> PropertyHandle = WeakPropertyHandle.Pin();
+												if (!PropertyHandle.IsValid())
+												{
+													return;
+												}
 
-											TSharedPtr<FStructOnScope> StructData = StructHandle->GetStructData();
+												TSharedPtr<IPropertyHandleStruct> StructHandle = PropertyHandle->AsStruct();
+												check(StructHandle.IsValid());
 
-											TInstancedStruct<FAnimNextTraitSharedData> Defaults;
-											Defaults.InitializeAsScriptStruct(InStruct);
-											
-											FRigVMControllerCompileBracketScope CompileBracketScope(Controller);
-											Controller->OpenUndoBracket(LOCTEXT("AddTraitTransaction", "Add Trait").ToString());
+												TSharedPtr<FStructOnScope> StructData = StructHandle->GetStructData();
 
-											for (TWeakObjectPtr<UAnimNextEdGraphNode> WeakNode : Nodes)
-											{
-												if (UAnimNextEdGraphNode* Node = WeakNode.Get())
+												TInstancedStruct<FAnimNextTraitSharedData> Defaults;
+												Defaults.InitializeAsScriptStruct(InStruct);
+												
+												FRigVMControllerCompileBracketScope CompileBracketScope(Controller);
+												Controller->OpenUndoBracket(LOCTEXT("AddTraitTransaction", "Add Trait").ToString());
+
+												for (TWeakObjectPtr<UAnimNextEdGraphNode> WeakNode : Nodes)
 												{
-													TArray<URigVMPin*> TraitPins = Node->GetModelNode()->GetTraitPins();
-													int32 SourcePinIndex = INDEX_NONE;
-													for (URigVMPin* Pin : TraitPins)
+													if (UAnimNextEdGraphNode* Node = WeakNode.Get())
 													{
-														TSharedPtr<FStructOnScope> TraitScope = Pin->GetTraitInstance();
-														if (TraitScope.IsValid()) 
+														TArray<URigVMPin*> TraitPins = Node->GetModelNode()->GetTraitPins();
+														int32 SourcePinIndex = INDEX_NONE;
+														for (URigVMPin* Pin : TraitPins)
 														{
-															const FRigVMTrait* VMTrait = (FRigVMTrait*)TraitScope->GetStructMemory();
-															if (VMTrait->GetTraitSharedDataStruct() == StructData->GetStruct())
+															TSharedPtr<FStructOnScope> TraitScope = Pin->GetTraitInstance();
+															if (TraitScope.IsValid()) 
 															{
-																SourcePinIndex = Pin->GetPinIndex();
-																break;
+																const FRigVMTrait* VMTrait = (FRigVMTrait*)TraitScope->GetStructMemory();
+																if (VMTrait->GetTraitSharedDataStruct() == StructData->GetStruct())
+																{
+																	SourcePinIndex = Pin->GetPinIndex();
+																	break;
+																}
 															}
 														}
-													}
 
-													FName NewTrait = Controller->AddTraitStruct(CastChecked<URigVMUnitNode>(Node->GetModelNode()), Defaults, FMath::Min(SourcePinIndex + 1, Node->GetModelNode()->GetPins().Num()), true, true);
-													if (NewTrait != NAME_None)
-													{
-														FRigVMNodeLayout NodeLayout = Node->GetModelNode()->GetNodeLayout();
-														UUAFGraphNodeTemplate::AddDefaultTraitPinsToLayout(InStruct, NodeLayout);
-														Controller->SetNodeLayout(Node->GetModelNodeName(), NodeLayout, true, true);
+														FName NewTrait = Controller->AddTraitStruct(CastChecked<URigVMUnitNode>(Node->GetModelNode()), Defaults, FMath::Min(SourcePinIndex + 1, Node->GetModelNode()->GetPins().Num()), true, true);
+														if (NewTrait != NAME_None)
+														{
+															FRigVMNodeLayout NodeLayout = Node->GetModelNode()->GetNodeLayout();
+															UUAFGraphNodeTemplate::AddDefaultTraitPinsToLayout(InStruct, NodeLayout);
+															Controller->SetNodeLayout(Node->GetModelNodeName(), NodeLayout, true, true);
+														}
 													}
 												}
-											}
 
-											Controller->CloseUndoBracket();
-										}))
-									], FText::GetEmpty(), true);
-							}));
+												Controller->CloseUndoBracket();
+											}))
+										], FText::GetEmpty(), true);
+								}));
 
-						MenuBuilder.AddSubMenu(
-							LOCTEXT("ReplaceLabel", "Replace..."),
-							LOCTEXT("ReplaceTooltip", "Replace this trait with another"),
-							FNewMenuDelegate::CreateLambda([this](FMenuBuilder& InMenuBuilder)
-							{
-								const FTraitRegistry& TraitRegistry = FTraitRegistry::Get();
-								const FTrait* Trait = TraitRegistry.Find(SharedDataStruct);
-								check(Trait);
-
-								FStructViewerModule& StructViewerModule = FModuleManager::LoadModuleChecked<FStructViewerModule>("StructViewer");
-
-								FStructViewerInitializationOptions Options;
-								Options.Mode = EStructViewerMode::StructPicker;
-								Options.NameTypeToDisplay = EStructViewerNameTypeToDisplay::DisplayName;
-								Options.StructFilter = MakeShared<FStructFilter>(Nodes, Trait->GetTraitMode());
-
-								InMenuBuilder.AddWidget(
-									SNew(SBox)
-									.WidthOverride(300.0f)
-									.HeightOverride(300.0f)
-									[
-										StructViewerModule.CreateStructViewer(Options, FOnStructPicked::CreateLambda([this](const UScriptStruct* InStruct)
-										{
-											FSlateApplication::Get().DismissAllMenus();
-
-											UAnimNextController* Controller = Cast<UAnimNextController>(WeakController.Get());
-											if (Controller == nullptr)
+							MenuBuilder.AddSubMenu(
+								LOCTEXT("ReplaceLabel", "Replace..."),
+								LOCTEXT("ReplaceTooltip", "Replace this trait with another"),
+								FNewMenuDelegate::CreateLambda([this](FMenuBuilder& InMenuBuilder)
+								{
+									const FTraitRegistry& TraitRegistry = FTraitRegistry::Get();
+									const FTrait* Trait = TraitRegistry.Find(SharedDataStruct);
+									check(Trait);
+
+									FStructViewerModule& StructViewerModule = FModuleManager::LoadModuleChecked<FStructViewerModule>("StructViewer");
+
+									FStructViewerInitializationOptions Options;
+									Options.Mode = EStructViewerMode::StructPicker;
+									Options.NameTypeToDisplay = EStructViewerNameTypeToDisplay::DisplayName;
+									Options.StructFilter = MakeShared<FStructFilter>(Nodes, Trait->GetTraitMode());
+
+									InMenuBuilder.AddWidget(
+										SNew(SBox)
+										.WidthOverride(300.0f)
+										.HeightOverride(300.0f)
+										[
+											StructViewerModule.CreateStructViewer(Options, FOnStructPicked::CreateLambda([this](const UScriptStruct* InStruct)
 											{
-												return;
-											}
+												FSlateApplication::Get().DismissAllMenus();
 
-											TSharedPtr<IPropertyHandle> PropertyHandle = WeakPropertyHandle.Pin();
-											if (!PropertyHandle.IsValid())
-											{
-												return;
-											}
+												UAnimNextController* Controller = Cast<UAnimNextController>(WeakController.Get());
+												if (Controller == nullptr)
+												{
+													return;
+												}
+
+												TSharedPtr<IPropertyHandle> PropertyHandle = WeakPropertyHandle.Pin();
+												if (!PropertyHandle.IsValid())
+												{
+													return;
+												}
 
-											TSharedPtr<IPropertyHandleStruct> StructHandle = PropertyHandle->AsStruct();
-											check(StructHandle.IsValid());
+												TSharedPtr<IPropertyHandleStruct> StructHandle = PropertyHandle->AsStruct();
+												check(StructHandle.IsValid());
 
-											TSharedPtr<FStructOnScope> StructData = StructHandle->GetStructData();
+												TSharedPtr<FStructOnScope> StructData = StructHandle->GetStructData();
 
-											TInstancedStruct<FAnimNextTraitSharedData> Defaults;
-											Defaults.InitializeAsScriptStruct(InStruct);
-											
-											FRigVMControllerCompileBracketScope CompileBracketScope(Controller);
-											Controller->OpenUndoBracket(LOCTEXT("ReplaceTraitTransaction", "Replace Trait").ToString());
+												TInstancedStruct<FAnimNextTraitSharedData> Defaults;
+												Defaults.InitializeAsScriptStruct(InStruct);
+												
+												FRigVMControllerCompileBracketScope CompileBracketScope(Controller);
+												Controller->OpenUndoBracket(LOCTEXT("ReplaceTraitTransaction", "Replace Trait").ToString());
 
-											for (TWeakObjectPtr<UAnimNextEdGraphNode> WeakNode : Nodes)
-											{
-												if (UAnimNextEdGraphNode* Node = WeakNode.Get())
+												for (TWeakObjectPtr<UAnimNextEdGraphNode> WeakNode : Nodes)
 												{
-													TArray<URigVMPin*> TraitPins = Node->GetModelNode()->GetTraitPins();
-													int32 SourcePinIndex = INDEX_NONE;
-													FName SourcePinName = NAME_None;
-													for (URigVMPin* Pin : TraitPins)
+													if (UAnimNextEdGraphNode* Node = WeakNode.Get())
 													{
-														TSharedPtr<FStructOnScope> TraitScope = Pin->GetTraitInstance();
-														if (TraitScope.IsValid()) 
+														TArray<URigVMPin*> TraitPins = Node->GetModelNode()->GetTraitPins();
+														int32 SourcePinIndex = INDEX_NONE;
+														FName SourcePinName = NAME_None;
+														for (URigVMPin* Pin : TraitPins)
 														{
-															const FRigVMTrait* VMTrait = (FRigVMTrait*)TraitScope->GetStructMemory();
-															if (VMTrait->GetTraitSharedDataStruct() == StructData->GetStruct())
+															TSharedPtr<FStructOnScope> TraitScope = Pin->GetTraitInstance();
+															if (TraitScope.IsValid()) 
 															{
-																SourcePinIndex = Pin->GetPinIndex();
-																SourcePinName = Pin->GetFName();
-																break;
+																const FRigVMTrait* VMTrait = (FRigVMTrait*)TraitScope->GetStructMemory();
+																if (VMTrait->GetTraitSharedDataStruct() == StructData->GetStruct())
+																{
+																	SourcePinIndex = Pin->GetPinIndex();
+																	SourcePinName = Pin->GetFName();
+																	break;
+																}
 															}
 														}
-													}
 
-													if (SourcePinIndex == INDEX_NONE || SourcePinName == NAME_None)
-													{
-														continue;
-													}
+														if (SourcePinIndex == INDEX_NONE || SourcePinName == NAME_None)
+														{
+															continue;
+														}
 
-													if (Controller->RemoveTraitByName(Node->GetModelNodeName(),SourcePinName, true, true))
-													{
-														FName NewTrait = Controller->AddTraitStruct(CastChecked<URigVMUnitNode>(Node->GetModelNode()), Defaults, SourcePinIndex, true, true);
-														if (NewTrait != NAME_None)
+														if (Controller->RemoveTraitByName(Node->GetModelNodeName(),SourcePinName, true, true))
 														{
-															FRigVMNodeLayout NodeLayout = Node->GetModelNode()->GetNodeLayout();
-															UUAFGraphNodeTemplate::AddDefaultTraitPinsToLayout(InStruct, NodeLayout);
-															Controller->SetNodeLayout(Node->GetModelNodeName(), NodeLayout, true, true);
+															FName NewTrait = Controller->AddTraitStruct(CastChecked<URigVMUnitNode>(Node->GetModelNode()), Defaults, SourcePinIndex, true, true);
+															if (NewTrait != NAME_None)
+															{
+																FRigVMNodeLayout NodeLayout = Node->GetModelNode()->GetNodeLayout();
+																UUAFGraphNodeTemplate::AddDefaultTraitPinsToLayout(InStruct, NodeLayout);
+																Controller->SetNodeLayout(Node->GetModelNodeName(), NodeLayout, true, true);
+															}
 														}
 													}
 												}
-											}
-
-											Controller->CloseUndoBracket();
-										}))
-									], FText::GetEmpty(), true);
-							}));
-
-						MenuBuilder.AddMenuEntry(
-							LOCTEXT("RemoveLabel", "Remove"),
-							LOCTEXT("RemoveTooltip", "Remove this trait"),
-							FSlateIcon(),
-							FExecuteAction::CreateLambda([this]()
-							{
-								FSlateApplication::Get().DismissAllMenus();
 
-								UAnimNextController* Controller = Cast<UAnimNextController>(WeakController.Get());
-								if (Controller == nullptr)
-								{
-									return;
-								}
+												Controller->CloseUndoBracket();
+											}))
+										], FText::GetEmpty(), true);
+								}));
 
-								TSharedPtr<IPropertyHandle> PropertyHandle = WeakPropertyHandle.Pin();
-								if (!PropertyHandle.IsValid())
+							MenuBuilder.AddMenuEntry(
+								LOCTEXT("RemoveLabel", "Remove"),
+								LOCTEXT("RemoveTooltip", "Remove this trait"),
+								FSlateIcon(),
+								FExecuteAction::CreateLambda([this]()
 								{
-									return;
-								}
+									FSlateApplication::Get().DismissAllMenus();
+
+									UAnimNextController* Controller = Cast<UAnimNextController>(WeakController.Get());
+									if (Controller == nullptr)
+									{
+										return;
+									}
+
+									TSharedPtr<IPropertyHandle> PropertyHandle = WeakPropertyHandle.Pin();
+									if (!PropertyHandle.IsValid())
+									{
+										return;
+									}
+
+									TSharedPtr<IPropertyHandleStruct> StructHandle = PropertyHandle->AsStruct();
+									check(StructHandle.IsValid());
+
+									TSharedPtr<FStructOnScope> StructData = StructHandle->GetStructData();
+									
+									FRigVMControllerCompileBracketScope CompileBracketScope(Controller);
+									Controller->OpenUndoBracket(LOCTEXT("RemoveTraitTransaction", "Remove Trait").ToString());
 
+									for (TWeakObjectPtr<UAnimNextEdGraphNode> WeakNode : Nodes)
+									{
+										if (UAnimNextEdGraphNode* Node = WeakNode.Get())
+										{
+											Controller->RemoveTraitByName(Node->GetModelNodeName(),StructData->GetStruct()->GetFName(), true, true);
+										}
+									}
+
+									Controller->CloseUndoBracket();
+								}));
+
+							TSharedPtr<IPropertyHandle> PropertyHandle = WeakPropertyHandle.Pin();
+							if (PropertyHandle.IsValid())
+							{
 								TSharedPtr<IPropertyHandleStruct> StructHandle = PropertyHandle->AsStruct();
 								check(StructHandle.IsValid());
 
 								TSharedPtr<FStructOnScope> StructData = StructHandle->GetStructData();
-								
-								FRigVMControllerCompileBracketScope CompileBracketScope(Controller);
-								Controller->OpenUndoBracket(LOCTEXT("RemoveTraitTransaction", "Remove Trait").ToString());
-
-								for (TWeakObjectPtr<UAnimNextEdGraphNode> WeakNode : Nodes)
+								if (StructData.IsValid() && StructData->GetStruct() && StructData->GetStruct()->HasMetaData(TEXT("Deprecated")))
 								{
-									if (UAnimNextEdGraphNode* Node = WeakNode.Get())
+									const FTrait* Trait = FTraitRegistry::Get().Find(CastChecked<UScriptStruct>(StructData->GetStruct()));
+									check(Trait);
+
+									FRigVMStructUpgradeInfo UpgradeInfo = Trait->GetUpgradeInfo();
+									if (UpgradeInfo.IsValid())
 									{
-										Controller->RemoveTraitByName(Node->GetModelNodeName(),StructData->GetStruct()->GetFName(), true, true);
+										MenuBuilder.AddMenuEntry(
+											LOCTEXT("UpgradeLabel", "Upgrade"),
+											LOCTEXT("UpgradeTooltip", "Upgrade this trait"),
+											FSlateIcon(),
+											FExecuteAction::CreateLambda([this, UpgradeInfo = MoveTemp(UpgradeInfo)]()
+											{
+												FSlateApplication::Get().DismissAllMenus();
+
+												UAnimNextController* Controller = Cast<UAnimNextController>(WeakController.Get());
+												if (Controller == nullptr)
+												{
+													return;
+												}
+
+												TSharedPtr<IPropertyHandle> PropertyHandle = WeakPropertyHandle.Pin();
+												if (!PropertyHandle.IsValid())
+												{
+													return;
+												}
+
+												TSharedPtr<IPropertyHandleStruct> StructHandle = PropertyHandle->AsStruct();
+												check(StructHandle.IsValid());
+
+												TSharedPtr<FStructOnScope> StructData = StructHandle->GetStructData();
+
+												FRigVMControllerCompileBracketScope CompileBracketScope(Controller);
+												Controller->OpenUndoBracket(LOCTEXT("UpgradeTraitTransaction", "Upgrade Trait").ToString());
+
+												for (TWeakObjectPtr<UAnimNextEdGraphNode> WeakNode : Nodes)
+												{
+													if (UAnimNextEdGraphNode* Node = WeakNode.Get())
+													{
+														Controller->UpgradeTraitByName(Node->GetModelNodeName(), StructData->GetStruct()->GetFName(), true, true);
+													}
+												}
+
+												Controller->CloseUndoBracket();
+											}));
 									}
 								}
+							}
 
-								Controller->CloseUndoBracket();
-							}));
-
-						return MenuBuilder.MakeWidget();
-					})
+							return MenuBuilder.MakeWidget();
+						})
+					]
 				]
 			];
 
diff --git a/Engine/Plugins/Experimental/UAF/UAFAnimGraph/Source/UAFAnimGraphUncookedOnly/Internal/AnimNextController.h b/Engine/Plugins/Experimental/UAF/UAFAnimGraph/Source/UAFAnimGraphUncookedOnly/Internal/AnimNextController.h
--- a/Engine/Plugins/Experimental/UAF/UAFAnimGraph/Source/UAFAnimGraphUncookedOnly/Internal/AnimNextController.h
+++ b/Engine/Plugins/Experimental/UAF/UAFAnimGraph/Source/UAFAnimGraphUncookedOnly/Internal/AnimNextController.h
@@ -48,6 +48,9 @@ public:
 	UFUNCTION(BlueprintCallable, Category = "UAF|Controller")
 	UAFANIMGRAPHUNCOOKEDONLY_API bool SetTraitPinIndex(FName InNodeName, FName InTraitInstanceName, int32 InNewPinIndex, bool bSetupUndoRedo = true, bool bPrintPythonCommand = false);
 
+	// Upgrade a trait if it needs to be
+	UAFANIMGRAPHUNCOOKEDONLY_API bool UpgradeTraitByName(FName InNodeName, FName InTraitInstanceName, bool bSetupUndoRedo = true, bool bPrintPythonCommand = false);
+
 	// Adds a unit node with a dynamic number of pins
 	URigVMUnitNode* AddUnitNodeWithPins(UScriptStruct* InScriptStruct, const FRigVMPinInfoArray& PinArray, const FName& InMethodName = TEXT("Execute"), const FVector2D& InPosition = FVector2D::ZeroVector, const FString& InNodeName = TEXT(""), bool bSetupUndoRedo = true, bool bPrintPythonCommand = false);
 };
diff --git a/Engine/Plugins/Experimental/UAF/UAFAnimGraph/Source/UAFAnimGraphUncookedOnly/Internal/Templates/GraphNodeTemplateRegistry.h b/Engine/Plugins/Experimental/UAF/UAFAnimGraph/Source/UAFAnimGraphUncookedOnly/Internal/Templates/GraphNodeTemplateRegistry.h
--- a/Engine/Plugins/Experimental/UAF/UAFAnimGraph/Source/UAFAnimGraphUncookedOnly/Internal/Templates/GraphNodeTemplateRegistry.h
+++ b/Engine/Plugins/Experimental/UAF/UAFAnimGraph/Source/UAFAnimGraphUncookedOnly/Internal/Templates/GraphNodeTemplateRegistry.h
@@ -7,6 +7,8 @@
 
 #define UE_API UAFANIMGRAPHUNCOOKEDONLY_API 
 
+struct FAnimNextParamType;
+
 namespace UE::UAF::AnimGraph::UncookedOnly
 {
 	class FAnimNextAnimGraphUncookedOnlyModule;
@@ -22,6 +24,7 @@ struct FGraphNodeTemplateInfo
 	FText Tooltip;
 	FText Category;
 	TArray<FTopLevelAssetPath> DragDropAssetTypes;
+	TArray<FAnimNextParamType> DragDropVariableTypes;
 };
 
 class FGraphNodeTemplateRegistry
@@ -33,6 +36,9 @@ public:
 	// Get all drag drop handlers registered to an asset type
 	UE_API static TConstArrayView<FGraphNodeTemplateInfo> GetDragDropHandlersForAsset(const FAssetData& InAssetData);
 
+	// Get all drag drop handlers registered to a variable type
+	UE_API static TConstArrayView<FGraphNodeTemplateInfo> GetDragDropHandlersForVariable(const FAnimNextParamType& InVariableType);
+
 private:
 	friend UE::UAF::AnimGraph::UncookedOnly::FAnimNextAnimGraphUncookedOnlyModule;
 
diff --git a/Engine/Plugins/Experimental/UAF/UAFAnimGraph/Source/UAFAnimGraphUncookedOnly/Private/AnimNextController.cpp b/Engine/Plugins/Experimental/UAF/UAFAnimGraph/Source/UAFAnimGraphUncookedOnly/Private/AnimNextController.cpp
--- a/Engine/Plugins/Experimental/UAF/UAFAnimGraph/Source/UAFAnimGraphUncookedOnly/Private/AnimNextController.cpp
+++ b/Engine/Plugins/Experimental/UAF/UAFAnimGraph/Source/UAFAnimGraphUncookedOnly/Private/AnimNextController.cpp
@@ -342,3 +342,135 @@ URigVMUnitNode* UAnimNextController::AddUnitNodeWithPins(UScriptStruct* InScript
 
 	return Node;
 }
+
+bool UAnimNextController::UpgradeTraitByName(FName InNodeName, FName InTraitInstanceName, bool bSetupUndoRedo, bool bPrintPythonCommand)
+{
+	if (!IsValidGraph())
+	{
+		return false;
+	}
+
+	if (!IsTransacting() && !IsGraphEditable())
+	{
+		return false;
+	}
+
+	if (InNodeName == NAME_None)
+	{
+		ReportError(TEXT("Invalid node name."));
+		return false;
+	}
+
+	const URigVMGraph* Graph = GetGraph();
+	check(Graph);
+
+	URigVMNode* Node = Graph->FindNodeByName(InNodeName);
+	if (Node == nullptr)
+	{
+		ReportError(TEXT("This graph does not contain a node with the provided name."));
+		return false;
+	}
+	
+	URigVMPin* TraitPin = Node->FindTrait(InTraitInstanceName);
+	if (TraitPin == nullptr)
+	{
+		ReportError(TEXT("The node does not contain a Trait with the provided name."));
+		return false;
+	}
+
+	// Avoid multiple VM recompilations for internal operations
+	FRigVMControllerCompileBracketScope CompileScope(this);
+
+	TMap<FString, FString> RedirectedPinPaths;
+	TMap<FString, FPinState> PinStates = GetPinStates(Node, true);
+	EjectAllInjectedNodes(Node, bSetupUndoRedo);
+
+	FRigVMBaseAction Action(this);
+	if(bSetupUndoRedo)
+	{
+		Action.SetTitle(TEXT("Upgrade trait"));
+		GetActionStack()->BeginAction(Action);
+	}
+
+	if (TSharedPtr<FStructOnScope> ScopedTrait = Node->GetTraitInstance(TraitPin->GetFName()))
+	{
+		const FRigVMTrait* OldTrait = (FRigVMTrait*)ScopedTrait->GetStructMemory();
+		if (const UScriptStruct* TraitSharedInstanceData = OldTrait->GetTraitSharedDataStruct())
+		{
+			UE::UAF::FTraitRegistry& TraitRegistry = UE::UAF::FTraitRegistry::Get();
+			if (const UE::UAF::FTrait* OldUAFTrait = TraitRegistry.Find(TraitSharedInstanceData))
+			{
+				FRigVMStructUpgradeInfo UpgradeInfo = OldUAFTrait->GetUpgradeInfo();
+				if (UpgradeInfo.IsValid())
+				{
+					if(UpgradeInfo.GetNewStruct()->IsChildOf(FAnimNextTraitSharedData::StaticStruct()))
+					{
+						if (!RemoveTraitByName(InNodeName, InTraitInstanceName, bSetupUndoRedo, bPrintPythonCommand))
+						{
+							if(bSetupUndoRedo)
+							{
+								GetActionStack()->CancelAction(Action);
+							}
+							ReportErrorf(TEXT("Unable to upgrade trait %s."), *InTraitInstanceName.ToString());
+							return false;
+						}
+
+						TInstancedStruct<FAnimNextTraitSharedData> Defaults(UpgradeInfo.GetNewStruct());
+						FName NewTraitName = AddTraitStruct(CastChecked<URigVMUnitNode>(Node), Defaults, TraitPin->GetPinIndex(), bSetupUndoRedo, bPrintPythonCommand);
+						if (NewTraitName.IsNone())
+						{
+							if(bSetupUndoRedo)
+							{
+								GetActionStack()->CancelAction(Action);
+							}
+							ReportErrorf(TEXT("Unable to upgrade trait %s."), *InTraitInstanceName.ToString());
+							return false;
+						}
+
+						URigVMPin* NewTraitPin = Node->FindTrait(NewTraitName);
+						check(NewTraitPin);
+
+						for(URigVMPin* Pin : NewTraitPin->GetSubPins())
+						{
+							const FString DefaultValue = UpgradeInfo.GetDefaultValueForPin(Pin->GetFName());
+							if(!DefaultValue.IsEmpty())
+							{
+								SetPinDefaultValue(Pin, DefaultValue, true, bSetupUndoRedo, false);
+
+								if(FPinState* PinState = PinStates.Find(Pin->GetPinPath()))
+								{
+									PinState->DefaultValue.Reset();
+								}
+							}
+						}
+						
+						// redirect pin state paths
+						for(TPair<FString, FPinState>& PinState : PinStates)
+						{
+							for(int32 TrueFalse = 0; TrueFalse < 2; TrueFalse++)
+							{
+								const FString RemappedInputPath = UpgradeInfo.RemapPin(PinState.Key, TrueFalse == 0, false);
+								if(RemappedInputPath != PinState.Key)
+								{
+									if(!RedirectedPinPaths.Contains(PinState.Key))
+									{
+										RedirectedPinPaths.Add(PinState.Key, RemappedInputPath);
+									}
+								}
+							}
+						}
+					}
+				}
+			}
+		}
+	}
+
+	ApplyPinStates(Node, PinStates, RedirectedPinPaths, bSetupUndoRedo);
+
+	if(bSetupUndoRedo)
+	{
+		GetActionStack()->EndAction(Action);
+	}
+	
+	return true;
+}
\ No newline at end of file
diff --git a/Engine/Plugins/Experimental/UAF/UAFAnimGraph/Source/UAFAnimGraphUncookedOnly/Private/Graph/AnimNextAnimationGraph_EditorData.cpp b/Engine/Plugins/Experimental/UAF/UAFAnimGraph/Source/UAFAnimGraphUncookedOnly/Private/Graph/AnimNextAnimationGraph_EditorData.cpp
--- a/Engine/Plugins/Experimental/UAF/UAFAnimGraph/Source/UAFAnimGraphUncookedOnly/Private/Graph/AnimNextAnimationGraph_EditorData.cpp
+++ b/Engine/Plugins/Experimental/UAF/UAFAnimGraph/Source/UAFAnimGraphUncookedOnly/Private/Graph/AnimNextAnimationGraph_EditorData.cpp
@@ -334,6 +334,13 @@ namespace UE::UAF::UncookedOnly::Private
 		for (int32 TraitIndex = 0; TraitIndex < InMapping.TraitEntries.Num(); ++TraitIndex)
 		{
 			const FTraitEntryMapping& Entry = InMapping.TraitEntries[TraitIndex];
+
+			const UScriptStruct* TraitStruct = Entry.Trait->GetTraitSharedDataStruct();
+			if (InSettings.bWarnAboutDeprecatedNodes && TraitStruct->HasMetaData(TEXT("Deprecated")))
+			{
+				InSettings.ASTSettings.Report(EMessageSeverity::Warning, InVMUnitNode, FText::Format(LOCTEXT("DeprecatedTraitWarning", "@@: The trait '{0}' is deprecated. Please remove or upgrade the trait."), TraitStruct->GetDisplayNameText()).ToString());
+			}
+
 			switch (Entry.Trait->GetTraitMode())
 			{
 			case ETraitMode::Base:
diff --git a/Engine/Plugins/Experimental/UAF/UAFAnimGraph/Source/UAFAnimGraphUncookedOnly/Private/Templates/GraphNodeTemplateRegistry.cpp b/Engine/Plugins/Experimental/UAF/UAFAnimGraph/Source/UAFAnimGraphUncookedOnly/Private/Templates/GraphNodeTemplateRegistry.cpp
--- a/Engine/Plugins/Experimental/UAF/UAFAnimGraph/Source/UAFAnimGraphUncookedOnly/Private/Templates/GraphNodeTemplateRegistry.cpp
+++ b/Engine/Plugins/Experimental/UAF/UAFAnimGraph/Source/UAFAnimGraphUncookedOnly/Private/Templates/GraphNodeTemplateRegistry.cpp
@@ -13,6 +13,7 @@
 #include "Kismet2/KismetEditorUtilities.h"
 #include "Blueprint/BlueprintSupport.h"
 #include "String/ParseTokens.h"
+#include "Param/ParamType.h"
 
 namespace UE::UAF
 {
@@ -24,6 +25,9 @@ namespace Private
 
 	// Map from asset type to drag-drop handlers for that asset
 	static TMap<FTopLevelAssetPath, TArray<FGraphNodeTemplateInfo>> GAssetDragDropHandlers;
+	
+	// Map from variable type to drag-drop handlers for that variable type
+	static TMap<FAnimNextParamType, TArray<FGraphNodeTemplateInfo>> GVariableDragDropHandlers;
 
 	// Make sure we dont initialize more than once
 	static bool bGIsTemplateRegistryInitialized = false;
@@ -179,7 +183,7 @@ void FGraphNodeTemplateRegistry::LazyInitialize()
 
 				FGraphNodeTemplateInfo NewInfo = GetTemplateInfoFromAsset(AssetData, ClassPath);
 
-				// Add new drag-drop handelrs
+				// Add new drag-drop handlers
 				if (NewInfo.DragDropAssetTypes.Num() > 0)
 				{
 					for (const FTopLevelAssetPath& AssetType : NewInfo.DragDropAssetTypes)
@@ -223,6 +227,7 @@ void FGraphNodeTemplateRegistry::LazyInitialize()
 		GraphNodeTemplateInfo.Tooltip = CDO->GetTooltipText();
 		GraphNodeTemplateInfo.Category = CDO->GetCategory();
 		GraphNodeTemplateInfo.DragDropAssetTypes = CDO->GetDragDropAssetTypes();
+		GraphNodeTemplateInfo.DragDropVariableTypes = CDO->GetDragDropVariableTypes();
 
 		if (GraphNodeTemplateInfo.DragDropAssetTypes.Num() > 0)
 		{
@@ -232,6 +237,15 @@ void FGraphNodeTemplateRegistry::LazyInitialize()
 				Handlers.Add(GraphNodeTemplateInfo);
 			}
 		}
+		
+		if (GraphNodeTemplateInfo.DragDropVariableTypes.Num() > 0)
+		{
+			for (const FAnimNextParamType& VariableType : GraphNodeTemplateInfo.DragDropVariableTypes)
+			{
+				TArray<FGraphNodeTemplateInfo>& Handlers = Private::GVariableDragDropHandlers.FindOrAdd(VariableType);
+				Handlers.Add(GraphNodeTemplateInfo);
+			}
+		}
 
 		Private::GAllTemplates.Add(ClassPath, MoveTemp(GraphNodeTemplateInfo));
 	}
@@ -256,6 +270,10 @@ void FGraphNodeTemplateRegistry::Shutdown()
 
 	FCoreUObjectDelegates::OnObjectPostCDOCompiled.Remove(Private::GOnObjectPostCDOCompiledHandle);
 
+	Private::GAllTemplates.Empty();
+	Private::GAssetDragDropHandlers.Empty();
+	Private::GVariableDragDropHandlers.Empty();
+
 	Private::bGIsTemplateRegistryInitialized = false;
 }
 
@@ -275,4 +293,14 @@ TConstArrayView<FGraphNodeTemplateInfo> FGraphNodeTemplateRegistry::GetDragDropH
 	return TConstArrayView<FGraphNodeTemplateInfo>();
 }
 
+TConstArrayView<FGraphNodeTemplateInfo> FGraphNodeTemplateRegistry::GetDragDropHandlersForVariable(const FAnimNextParamType& InVariableType)
+{
+	LazyInitialize();
+	if (TArray<FGraphNodeTemplateInfo>* FoundHandlers = Private::GVariableDragDropHandlers.Find(InVariableType))
+	{
+		return *FoundHandlers;
+	}
+	return TConstArrayView<FGraphNodeTemplateInfo>();
+}
+
 }
diff --git a/Engine/Plugins/Experimental/UAF/UAFAnimGraph/Source/UAFAnimGraphUncookedOnly/Private/Templates/UAFGraphNodeTemplate.cpp b/Engine/Plugins/Experimental/UAF/UAFAnimGraph/Source/UAFAnimGraphUncookedOnly/Private/Templates/UAFGraphNodeTemplate.cpp
--- a/Engine/Plugins/Experimental/UAF/UAFAnimGraph/Source/UAFAnimGraphUncookedOnly/Private/Templates/UAFGraphNodeTemplate.cpp
+++ b/Engine/Plugins/Experimental/UAF/UAFAnimGraph/Source/UAFAnimGraphUncookedOnly/Private/Templates/UAFGraphNodeTemplate.cpp
@@ -75,6 +75,10 @@ void UUAFGraphNodeTemplate::HandleAssetDropped_Implementation(UAnimNextControlle
 	Controller->CloseUndoBracket();
 }
 
+void UUAFGraphNodeTemplate::HandleVariableDropped_Implementation(UAnimNextController* Controller, URigVMUnitNode* Node, FAnimNextVariableReference Variable) const
+{
+}
+
 FRigVMNodeLayout UUAFGraphNodeTemplate::GetDefaultCategoryLayout() const
 {
 	FRigVMNodeLayout Layout;
diff --git a/Engine/Plugins/Experimental/UAF/UAFAnimGraph/Source/UAFAnimGraphUncookedOnly/Public/Templates/UAFGraphNodeTemplate.h b/Engine/Plugins/Experimental/UAF/UAFAnimGraph/Source/UAFAnimGraphUncookedOnly/Public/Templates/UAFGraphNodeTemplate.h
--- a/Engine/Plugins/Experimental/UAF/UAFAnimGraph/Source/UAFAnimGraphUncookedOnly/Public/Templates/UAFGraphNodeTemplate.h
+++ b/Engine/Plugins/Experimental/UAF/UAFAnimGraph/Source/UAFAnimGraphUncookedOnly/Public/Templates/UAFGraphNodeTemplate.h
@@ -10,6 +10,7 @@
 #include "Textures/SlateIcon.h"
 #include "UAFGraphNodeTemplate.generated.h"
 
+struct FAnimNextParamType;
 struct FAnimNextTraitSharedData;
 class UAnimNextTraitStackUnitNode;
 class URigVMController;
@@ -118,6 +119,11 @@ public:
 	UE_API void HandleAssetDropped(UAnimNextController* Controller, URigVMUnitNode* Node, UObject* Asset) const;
 	UE_API virtual void HandleAssetDropped_Implementation(UAnimNextController* Controller, URigVMUnitNode* Node, UObject* Asset) const;
 
+	// Called in editor when a variable is drag/dropped. Called after ConfigureNewNode.
+	UFUNCTION(BlueprintNativeEvent, Category = "UAF|Template")
+	UE_API void HandleVariableDropped(UAnimNextController* Controller, URigVMUnitNode* Node, FAnimNextVariableReference Variable) const;
+	UE_API virtual void HandleVariableDropped_Implementation(UAnimNextController* Controller, URigVMUnitNode* Node, FAnimNextVariableReference Variable) const;
+	
 	// Called in editor when a pin default value is changed, e.g. to be used to reconfigure the nodes appearance
 	UFUNCTION(BlueprintNativeEvent, Category = "UAF|Template")
 	UE_API void HandlePinDefaultValueChanged(UAnimNextController* Controller, URigVMPin* Pin) const;
@@ -149,7 +155,13 @@ public:
 	{
 		return DragDropAssetTypes;
 	}
-
+	
+	// Get all the types of variable this template handles for drag-drop
+	TConstArrayView<FAnimNextParamType> GetDragDropVariableTypes() const
+	{
+		return DragDropVariableTypes;
+	}
+	
 	// Utility function used to create and configure a new node according to a template
 	UE_API URigVMUnitNode* CreateNewNode(UAnimNextController* Controller, const FVector2D& InLocation) const;
 
@@ -206,6 +218,10 @@ protected:
 	UPROPERTY(EditAnywhere, BlueprintReadOnly, AssetRegistrySearchable, Category="Template")
 	TArray<TSubclassOf<UObject>> DragDropAssetTypes;
 
+	// Variable types that we handle for drag-drop
+	UPROPERTY(EditAnywhere, BlueprintReadOnly, AssetRegistrySearchable, Category="Template")
+	TArray<FAnimNextParamType> DragDropVariableTypes;
+
 	UPROPERTY(EditAnywhere, BlueprintReadOnly, Category="Template")
 	FRigVMNodeLayout NodeLayout;
 };
