UAF: set default value for mesh space flag on make dynamic additive trait

#Synced-CL 50765472

diff --git a/Engine/Plugins/Experimental/UAF/UAFAnimGraph/Source/UAFAnimGraph/Public/Traits/MakeDynamicAdditiveTraitData.h b/Engine/Plugins/Experimental/UAF/UAFAnimGraph/Source/UAFAnimGraph/Public/Traits/MakeDynamicAdditiveTraitData.h
--- a/Engine/Plugins/Experimental/UAF/UAFAnimGraph/Source/UAFAnimGraph/Public/Traits/MakeDynamicAdditiveTraitData.h
+++ b/Engine/Plugins/Experimental/UAF/UAFAnimGraph/Source/UAFAnimGraph/Public/Traits/MakeDynamicAdditiveTraitData.h
@@ -26,7 +26,7 @@ namespace UE::UAF
 		
 		// Do additive delta calculation in mesh space
 		UPROPERTY(EditAnywhere, Category="Default")
-		bool bMeshSpaceAdditive;
+		bool bMeshSpaceAdditive = false;
 		
 		// Latent pin support boilerplate
 		#define TRAIT_LATENT_PROPERTIES_ENUMERATOR(GeneratorMacro) \
