UAF: GC improvements for standalone graphs

GC for traits is now handled by graph instances, rather than a component on the root graph instance. This allows graphs to exist standalone and still correctly respond to GC, while keeping GC non-recursive.

#Synced-CL 50778421

diff --git a/Engine/Plugins/Experimental/UAF/UAFAnimGraph/Source/UAFAnimGraph/Private/Graph/AnimNextGraphInstance.cpp b/Engine/Plugins/Experimental/UAF/UAFAnimGraph/Source/UAFAnimGraph/Private/Graph/AnimNextGraphInstance.cpp
--- a/Engine/Plugins/Experimental/UAF/UAFAnimGraph/Source/UAFAnimGraph/Private/Graph/AnimNextGraphInstance.cpp
+++ b/Engine/Plugins/Experimental/UAF/UAFAnimGraph/Source/UAFAnimGraph/Private/Graph/AnimNextGraphInstance.cpp
@@ -9,12 +9,12 @@
 #include "Graph/AnimNextGraphContextData.h"
 #include "Graph/AnimNextGraphLatentPropertiesContextData.h"
 #include "TraitCore/ExecutionContext.h"
-#include "Graph/GC_GraphInstanceComponent.h"
 #include "Graph/RigUnit_AnimNextShimRoot.h"
 #include "Graph/RigVMTrait_AnimNextPublicVariables.h"
 #include "Logging/StructuredLog.h"
 #include "Module/AnimNextModule.h"
 #include "Module/AnimNextModuleInstance.h"
+#include "TraitInterfaces/IGarbageCollection.h"
 
 #include UE_INLINE_GENERATED_CPP_BY_NAME(AnimNextGraphInstance)
 
@@ -49,6 +49,7 @@ void FAnimNextGraphInstance::Release()
 	GraphInstancePtr.Reset();
 	ModuleInstance = nullptr;
 	RootGraphInstance = nullptr;
+	GCHandler.TraitsWithReferences.Reset();
 	ReleaseComponents();
 	Asset = nullptr;
 }
@@ -195,6 +196,7 @@ void FAnimNextGraphInstance::Freeze()
 
 	GraphInstancePtr.Reset();
 	ReleaseComponents();
+	GCHandler.TraitsWithReferences.Reset();
 	bHasUpdatedOnce = false;
 }
 
@@ -228,3 +230,56 @@ void FAnimNextGraphInstance::Thaw()
 }
 
 #endif
+
+void FAnimNextGraphInstance::RegisterWithGC(const UE::UAF::FWeakTraitPtr& InTraitPtr)
+{
+	GCHandler.TraitsWithReferences.Add(InTraitPtr);
+}
+
+void FAnimNextGraphInstance::UnregisterWithGC(const UE::UAF::FWeakTraitPtr& InTraitPtr)
+{
+	using namespace UE::UAF;
+	
+	const int32 EntryIndex = GCHandler.TraitsWithReferences.IndexOfByPredicate(
+		[&InTraitPtr](const FWeakTraitPtr& TraitPtr)
+		{
+			return TraitPtr == InTraitPtr;
+		});
+
+	if (ensure(EntryIndex != INDEX_NONE))
+	{
+		GCHandler.TraitsWithReferences.RemoveAtSwap(EntryIndex);
+	}
+}
+
+void FAnimNextGraphInstance::AddStructReferencedObjects(FReferenceCollector& Collector)
+{
+	GCHandler.AddReferencedObjects(Collector);
+}
+
+void FAnimNextGraphInstance::FGCHandler::AddReferencedObjects(FReferenceCollector& Collector)
+{
+	using namespace UE::UAF;
+
+	FExecutionContext Context;
+	FTraitStackBinding TraitStack;
+	TTraitBinding<IGarbageCollection> GCTrait;
+
+	// TODO: If we kept the entries sorted by graph instance, we could re-use the execution context
+	for (const FWeakTraitPtr& TraitPtr : TraitsWithReferences)
+	{
+		Context.BindTo(TraitPtr);
+
+		if (Context.GetStack(TraitPtr, TraitStack))
+		{
+			check(TraitStack.GetInterface(GCTrait));
+
+			GCTrait.AddReferencedObjects(Context, Collector);
+		}
+	}
+}
+
+FString FAnimNextGraphInstance::FGCHandler::GetReferencerName() const
+{
+	return TEXT("UAF Graph Instance");
+}
diff --git a/Engine/Plugins/Experimental/UAF/UAFAnimGraph/Source/UAFAnimGraph/Private/Graph/GC_GraphInstanceComponent.cpp b/Engine/Plugins/Experimental/UAF/UAFAnimGraph/Source/UAFAnimGraph/Private/Graph/GC_GraphInstanceComponent.cpp
deleted file mode 100644
diff --git a/Engine/Plugins/Experimental/UAF/UAFAnimGraph/Source/UAFAnimGraph/Private/Graph/GC_GraphInstanceComponent.h b/Engine/Plugins/Experimental/UAF/UAFAnimGraph/Source/UAFAnimGraph/Private/Graph/GC_GraphInstanceComponent.h
deleted file mode 100644
diff --git a/Engine/Plugins/Experimental/UAF/UAFAnimGraph/Source/UAFAnimGraph/Private/TraitCore/ExecutionContext.cpp b/Engine/Plugins/Experimental/UAF/UAFAnimGraph/Source/UAFAnimGraph/Private/TraitCore/ExecutionContext.cpp
--- a/Engine/Plugins/Experimental/UAF/UAFAnimGraph/Source/UAFAnimGraph/Private/TraitCore/ExecutionContext.cpp
+++ b/Engine/Plugins/Experimental/UAF/UAFAnimGraph/Source/UAFAnimGraph/Private/TraitCore/ExecutionContext.cpp
@@ -13,6 +13,7 @@
 #include "Graph/AnimNextAnimationGraph.h"
 #include "Graph/AnimNextGraphInstance.h"
 #include "Components/SceneComponent.h"
+#include "Graph/UAFOutputPoseComponent.h"
 #include "Module/AnimNextModuleInstance.h"
 #include "Logging/StructuredLog.h"
 
@@ -196,10 +197,7 @@ namespace UE::UAF
 			
 			if (TraitDesc->IsGCAware())
 			{
-				const IGarbageCollection* GCImpl = static_cast<const IGarbageCollection*>(Trait->GetTraitInterface(IGarbageCollection::InterfaceUID));
-				ensure(GCImpl);
-				
-				GCImpl->RegisterWithGC(*this, Binding);
+				GraphInstance.RegisterWithGC(Binding.GetTraitPtr());
 			}
 
 			Trait->ConstructTraitInstance(*this, Binding);
@@ -296,13 +294,11 @@ namespace UE::UAF
 			// Auto-register this trait with garbage collection
 			if (TraitDesc->IsGCAware())
 			{
-				const IGarbageCollection* GCImpl = static_cast<const IGarbageCollection*>(Trait->GetTraitInterface(IGarbageCollection::InterfaceUID));
-				if (ensure(GCImpl))
-				{
-					GCImpl->RegisterWithGC(*this, Binding);
-				}
+				const FWeakTraitPtr TraitPtr = Binding.GetTraitPtr();
+				FAnimNextGraphInstance& Instance = TraitPtr.GetNodeInstance()->GetOwner();
+				Instance.RegisterWithGC(TraitPtr);
 			}
-            			
+
 			Trait->ConstructTraitInstance(*this, Binding);
 		}
 
@@ -370,11 +366,9 @@ namespace UE::UAF
 			// Auto-unregister this trait with garbage collection
 			if (TraitDesc->IsGCAware())
 			{
-				const IGarbageCollection* GCImpl = static_cast<const IGarbageCollection*>(Trait->GetTraitInterface(IGarbageCollection::InterfaceUID));
-				if (ensure(GCImpl))
-				{
-					GCImpl->UnregisterWithGC(*this, Binding);
-				}
+				const FWeakTraitPtr TraitPtr = Binding.GetTraitPtr();
+				FAnimNextGraphInstance& Instance = TraitPtr.GetNodeInstance()->GetOwner();
+				Instance.UnregisterWithGC(TraitPtr);
 			}
 			
 			Trait->DestructTraitInstance(*this, Binding);
diff --git a/Engine/Plugins/Experimental/UAF/UAFAnimGraph/Source/UAFAnimGraph/Private/TraitInterfaces/IGarbageColllection.cpp b/Engine/Plugins/Experimental/UAF/UAFAnimGraph/Source/UAFAnimGraph/Private/TraitInterfaces/IGarbageColllection.cpp
--- a/Engine/Plugins/Experimental/UAF/UAFAnimGraph/Source/UAFAnimGraph/Private/TraitInterfaces/IGarbageColllection.cpp
+++ b/Engine/Plugins/Experimental/UAF/UAFAnimGraph/Source/UAFAnimGraph/Private/TraitInterfaces/IGarbageColllection.cpp
@@ -3,7 +3,7 @@
 #include "TraitInterfaces/IGarbageCollection.h"
 
 #include "TraitCore/ExecutionContext.h"
-#include "Graph/GC_GraphInstanceComponent.h"
+#include "TraitCore/NodeInstance.h"
 
 namespace UE::UAF
 {
@@ -22,18 +22,6 @@ namespace UE::UAF
 	}
 #endif // WITH_EDITOR
 
-	void IGarbageCollection::RegisterWithGC(const FExecutionContext& Context, const FTraitBinding& Binding)
-	{
-		FGCGraphInstanceComponent& Component = Context.GetComponent<FGCGraphInstanceComponent>();
-		Component.Register(Binding.GetTraitPtr());
-	}
-
-	void IGarbageCollection::UnregisterWithGC(const FExecutionContext& Context, const FTraitBinding& Binding)
-	{
-		FGCGraphInstanceComponent& Component = Context.GetComponent<FGCGraphInstanceComponent>();
-		Component.Unregister(Binding.GetTraitPtr());
-	}
-
 	void IGarbageCollection::AddReferencedObjects(const FExecutionContext& Context, const TTraitBinding<IGarbageCollection>& Binding, FReferenceCollector& Collector) const
 	{
 		TTraitBinding<IGarbageCollection> SuperBinding;
diff --git a/Engine/Plugins/Experimental/UAF/UAFAnimGraph/Source/UAFAnimGraph/Public/Graph/AnimNextGraphInstance.h b/Engine/Plugins/Experimental/UAF/UAFAnimGraph/Source/UAFAnimGraph/Public/Graph/AnimNextGraphInstance.h
--- a/Engine/Plugins/Experimental/UAF/UAFAnimGraph/Source/UAFAnimGraph/Public/Graph/AnimNextGraphInstance.h
+++ b/Engine/Plugins/Experimental/UAF/UAFAnimGraph/Source/UAFAnimGraph/Public/Graph/AnimNextGraphInstance.h
@@ -103,6 +103,8 @@ struct FAnimNextGraphInstance : public FUAFAssetInstance
 	// Called each time the graph updates to mark the instance as updated
 	UE_API void MarkAsUpdated();
 
+	void AddStructReferencedObjects(FReferenceCollector& Collector);
+
 private:
 	// Executes a list of latent RigVM pins and writes the result into the destination pointer (latent handle offsets are using the destination as base)
 	// When frozen, latent handles that can freeze are skipped, all others will execute
@@ -121,6 +123,25 @@ private:
 	UE_API void Thaw();
 #endif
 
+	// Registers the provided trait with the GC system
+	// Once registered, IGarbageCollection::AddReferencedObjects will be called on it during GC
+	void RegisterWithGC(const UE::UAF::FWeakTraitPtr& InTraitPtr);
+	
+	// Unregisters the provided trait from the GC system
+	void UnregisterWithGC(const UE::UAF::FWeakTraitPtr& InTraitPtr);
+
+private:
+	// GC handler. This is nested to avoid having to bubble the polymorphism it introduces up to super-structs
+	struct FGCHandler : public FGCObject
+	{
+		// FGCObject implementation
+		virtual void AddReferencedObjects(FReferenceCollector& Collector) override;
+		virtual FString GetReferencerName() const override;
+
+		// List of trait handles that contain UObject references and implement IGarbageCollection
+		TArray<UE::UAF::FWeakTraitPtr> TraitsWithReferences;
+	} GCHandler;
+
 	// The entry point in Graph that this instance corresponds to 
 	FName EntryPoint;
 
@@ -160,6 +181,7 @@ struct TStructOpsTypeTraits<FAnimNextGraphInstance> : public TStructOpsTypeTrait
 	enum
 	{
 		WithCopy = false,
+		WithAddStructReferencedObjects = true,
 	};
 };
 
diff --git a/Engine/Plugins/Experimental/UAF/UAFAnimGraph/Source/UAFAnimGraph/Public/TraitInterfaces/IGarbageCollection.h b/Engine/Plugins/Experimental/UAF/UAFAnimGraph/Source/UAFAnimGraph/Public/TraitInterfaces/IGarbageCollection.h
--- a/Engine/Plugins/Experimental/UAF/UAFAnimGraph/Source/UAFAnimGraph/Public/TraitInterfaces/IGarbageCollection.h
+++ b/Engine/Plugins/Experimental/UAF/UAFAnimGraph/Source/UAFAnimGraph/Public/TraitInterfaces/IGarbageCollection.h
@@ -31,16 +31,6 @@ namespace UE::UAF
 		UE_API virtual const FText& GetDisplayName() const override;
 		UE_API virtual const FText& GetDisplayShortName() const override;
 #endif // WITH_EDITOR
-
-	private:
-		friend struct FExecutionContext;
-	
-		// Registers the provided binding for GC callback
-		// Once registered, AddReferencedObjects is called during GC to collect references
-		static UE_API void RegisterWithGC(const FExecutionContext& Context, const FTraitBinding& Binding);
-
-		// Unregisters the provided binding from GC callback
-		static UE_API void UnregisterWithGC(const FExecutionContext& Context, const FTraitBinding& Binding);
 	};
 
 	/**
diff --git a/Engine/Plugins/Experimental/UAF/UAFTestSuites/Source/UAFAnimGraphTestSuite/Private/GCInstanceDataGCTest.cpp b/Engine/Plugins/Experimental/UAF/UAFTestSuites/Source/UAFAnimGraphTestSuite/Private/GCInstanceDataGCTest.cpp
--- a/Engine/Plugins/Experimental/UAF/UAFTestSuites/Source/UAFAnimGraphTestSuite/Private/GCInstanceDataGCTest.cpp
+++ b/Engine/Plugins/Experimental/UAF/UAFTestSuites/Source/UAFAnimGraphTestSuite/Private/GCInstanceDataGCTest.cpp
@@ -166,7 +166,7 @@ bool FAnimationUAFRuntimeTest_InstanceDataGC::RunTest(const FString& InParameter
 		TSharedPtr<FAnimNextGraphInstance> GraphInstance = AnimationGraph->AllocateInstance();
 
 		// We place the graph instance in a UObject so that its contained within something the reference collector is aware of
-		// so that FGCGraphInstanceComponent.AddStructReferencedObjects is called. (This isn't strictly needed as it won't traverse
+		// so that FAnimNextGraphInstance.AddStructReferencedObjects is called. (This isn't strictly needed as it won't traverse
 		// our test trait anyway as it hasn't implemented IGarbageCollection) but its to match parity with the other GC test
 		TObjectPtr<UGraphInstanceHolder> GraphInstanceHolder = NewObject<UGraphInstanceHolder>();
 		GraphInstanceHolder->AddToRoot();
diff --git a/Engine/Plugins/Experimental/UAF/UAFTestSuites/Source/UAFAnimGraphTestSuite/Private/GCInstanceDataMissingGCTest.cpp b/Engine/Plugins/Experimental/UAF/UAFTestSuites/Source/UAFAnimGraphTestSuite/Private/GCInstanceDataMissingGCTest.cpp
--- a/Engine/Plugins/Experimental/UAF/UAFTestSuites/Source/UAFAnimGraphTestSuite/Private/GCInstanceDataMissingGCTest.cpp
+++ b/Engine/Plugins/Experimental/UAF/UAFTestSuites/Source/UAFAnimGraphTestSuite/Private/GCInstanceDataMissingGCTest.cpp
@@ -151,7 +151,7 @@ bool FAnimationUAFRuntimeTest_InstanceDataGCMissing::RunTest(const FString& InPa
 		TSharedPtr<FAnimNextGraphInstance> GraphInstance = AnimationGraph->AllocateInstance();
 
 		// We place the graph instance in a UObject so that its contained within something the reference collector is aware of
-		// so that FGCGraphInstanceComponent.AddStructReferencedObjects is called. (This isn't strictly needed as it won't traverse
+		// so that FAnimNextGraphInstance.AddReferencedObjects is called. (This isn't strictly needed as it won't traverse
 		// our test trait anyway as it hasn't implemented IGarbageCollection) but its to match parity with the other GC test
 		TObjectPtr<UGraphInstanceHolder> GraphInstanceHolder = NewObject<UGraphInstanceHolder>();
 		GraphInstanceHolder->AddToRoot();
