UAF: Remove debug name overrides

Its more useful to see the actual assets than their type in rewind debugger

#Synced-CL 50814747

diff --git a/Engine/Plugins/Experimental/UAF/UAFAnimGraph/Source/UAFAnimGraph/Private/AnimNextAnimGraphModule.cpp b/Engine/Plugins/Experimental/UAF/UAFAnimGraph/Source/UAFAnimGraph/Private/AnimNextAnimGraphModule.cpp
--- a/Engine/Plugins/Experimental/UAF/UAFAnimGraph/Source/UAFAnimGraph/Private/AnimNextAnimGraphModule.cpp
+++ b/Engine/Plugins/Experimental/UAF/UAFAnimGraph/Source/UAFAnimGraph/Private/AnimNextAnimGraphModule.cpp
@@ -173,12 +173,6 @@ void FAnimNextAnimGraphModule::RegisterGraphFactories()
 		FAnimNextFactoryParams Params;
 		Params.AddTraitStruct<FSequencePlayerData>(ETraitVariableMapping::All, 0);
 		Params.AddTraitStruct<FNotifyDispatcherData>(ETraitVariableMapping::None, 0);
-#if UE_TRACE_ENABLED
-		Params.AddInitializeTask([](const UE::UAF::FInstanceTaskContext& InContext)
-		{
-			InContext.GetAssetInstance().SetDebugName(TEXT("AnimSequence"));
-		});
-#endif
 
 		FAnimGraphFactory::RegisterAsset<FUAFGraphFactoryAsset_Animation>(
 			MoveTemp(Params),
@@ -197,12 +191,6 @@ void FAnimNextAnimGraphModule::RegisterGraphFactories()
 		// Register UBlendSpace 'player'
 		FAnimNextFactoryParams Params;
 		Params.AddTraitStruct<FBlendSpacePlayerData>(ETraitVariableMapping::All, 0);
-#if UE_TRACE_ENABLED
-		Params.AddInitializeTask([](const UE::UAF::FInstanceTaskContext& InContext)
-		{
-			InContext.GetAssetInstance().SetDebugName(TEXT("BlendSpace"));
-		});
-#endif
 
 		FAnimGraphFactory::RegisterAsset<FUAFGraphFactoryAsset_BlendSpace>(
 			MoveTemp(Params),
