UAF: Fix doubled-up variables in factory graphs

#Synced-CL 50814747

diff --git a/Engine/Plugins/Experimental/UAF/UAFAnimGraph/Source/UAFAnimGraph/Private/Factory/AnimNextSimpleAnimGraphBuilder.cpp b/Engine/Plugins/Experimental/UAF/UAFAnimGraph/Source/UAFAnimGraph/Private/Factory/AnimNextSimpleAnimGraphBuilder.cpp
--- a/Engine/Plugins/Experimental/UAF/UAFAnimGraph/Source/UAFAnimGraph/Private/Factory/AnimNextSimpleAnimGraphBuilder.cpp
+++ b/Engine/Plugins/Experimental/UAF/UAFAnimGraph/Source/UAFAnimGraph/Private/Factory/AnimNextSimpleAnimGraphBuilder.cpp
@@ -351,7 +351,10 @@ bool FAnimNextSimpleAnimGraphBuilder::Validate() const
 
 void FAnimNextSimpleAnimGraphBuilder::AutoMapTraitHelper(int32 InStackIndex, const UScriptStruct* TraitStruct)
 {
-	VariableStructs.Add(TraitStruct);
+	if (!VariableStructs.Contains(TraitStruct))
+	{
+		VariableStructs.Add(TraitStruct);
+	}
 
 	const int32 TraitIndex = Stacks[InStackIndex].TraitDescs.IndexOfByPredicate([TraitStruct](const FAnimNextSimpleAnimGraphBuilderTraitDesc& InDesc)
 	{
