Add missing component structs when using simple anim graph builder
* Any uses of .AddComponent<FMyComponent> won't get added otherwise

#Synced-CL 50860160

diff --git a/Engine/Plugins/Experimental/UAF/UAFAnimGraph/Source/UAFAnimGraph/Private/Factory/AnimNextSimpleAnimGraphBuilder.cpp b/Engine/Plugins/Experimental/UAF/UAFAnimGraph/Source/UAFAnimGraph/Private/Factory/AnimNextSimpleAnimGraphBuilder.cpp
--- a/Engine/Plugins/Experimental/UAF/UAFAnimGraph/Source/UAFAnimGraph/Private/Factory/AnimNextSimpleAnimGraphBuilder.cpp
+++ b/Engine/Plugins/Experimental/UAF/UAFAnimGraph/Source/UAFAnimGraph/Private/Factory/AnimNextSimpleAnimGraphBuilder.cpp
@@ -23,6 +23,7 @@ bool FAnimNextSimpleAnimGraphBuilder::Build(UE::UAF::FAnimGraphBuilderContext& I
 		return false;
 	}
 
+	InContext.ComponentStructs = ComponentStructs;
 	InContext.VariableStructs = VariableStructs;
 
 	// Index variables for mapping lookup
