Fix for multiple copies of variable instances showing up in rewing debugger for graphs

#Synced-CL 51031612

#Synced-CL 51031612

diff --git a/Engine/Plugins/Experimental/UAF/UAFAnimGraph/Source/UAFAnimGraph/Private/Graph/UAFAnimGraphTrace.cpp b/Engine/Plugins/Experimental/UAF/UAFAnimGraph/Source/UAFAnimGraph/Private/Graph/UAFAnimGraphTrace.cpp
--- a/Engine/Plugins/Experimental/UAF/UAFAnimGraph/Source/UAFAnimGraph/Private/Graph/UAFAnimGraphTrace.cpp
+++ b/Engine/Plugins/Experimental/UAF/UAFAnimGraph/Source/UAFAnimGraph/Private/Graph/UAFAnimGraphTrace.cpp
@@ -57,6 +57,7 @@ namespace UE::UAF
 				if (!TracedInstances.Contains(Graph.GetUniqueId()))
 				{
 					TRACE_UAF_VARIABLES(&Graph, Graph.GetModuleInstance() ? Graph.GetModuleInstance()->GetObject() : nullptr);
+					TracedInstances.Add(Graph.GetUniqueId());
 				}
 
 				FTraitStackBinding TraitStack;
