Remove OnBecomeRelevant tag from FAnimNextSequencePlayerTraitSharedData.

#Synced-CL 51152995

diff --git a/Engine/Plugins/Experimental/UAF/UAFAnimGraph/Source/UAFAnimGraph/Public/Traits/SequencePlayerTraitData.h b/Engine/Plugins/Experimental/UAF/UAFAnimGraph/Source/UAFAnimGraph/Public/Traits/SequencePlayerTraitData.h
--- a/Engine/Plugins/Experimental/UAF/UAFAnimGraph/Source/UAFAnimGraph/Public/Traits/SequencePlayerTraitData.h
+++ b/Engine/Plugins/Experimental/UAF/UAFAnimGraph/Source/UAFAnimGraph/Public/Traits/SequencePlayerTraitData.h
@@ -19,7 +19,7 @@ struct FAnimNextSequencePlayerTraitSharedData : public FAnimNextTraitSharedData
 	TObjectPtr<const UAnimSequence> AnimSequence;
 
 	/** The play rate multiplier at which this sequence plays. */
-	UPROPERTY(EditAnywhere, Category = "Sequence Player", meta = (OnBecomeRelevant))
+	UPROPERTY(EditAnywhere, Category = "Sequence Player")
 	float PlayRate = 1.0f;
 
 	/** The time at which we should start playing this sequence. */
