Update UAF blend profiles
* Update all existing uses of UHierarchyTable in UAF to now use UUAFBlendProfile and UUAFBlendMask.
* Motivation for dedicated UAF implementations is in order to support the abstract heirarchy without needing to rework the engine UBlendProfile.
* Motivation for splitting UBlendProfile up onto blend profiles (time & weight) and blend masks is due to their different uses: blend mask is for defining *what* is blended, blend profiles is for controlling *how* it blends. Having these two different concepts under one unbrella class was confusing.

* Fix time & weight blend profiles handling in blend stack.
* Update state tree to support transitioning with a blend profile.
* Update layer stack to use blend masks (also supporting blend profiles is TODO).

#Synced-CL 51369688

diff --git a/Engine/Plugins/Experimental/Animation/HierarchyTable/Source/Editor/Private/FloatColumn.cpp b/Engine/Plugins/Experimental/Animation/HierarchyTable/Source/Editor/Private/FloatColumn.cpp
--- a/Engine/Plugins/Experimental/Animation/HierarchyTable/Source/Editor/Private/FloatColumn.cpp
+++ b/Engine/Plugins/Experimental/Animation/HierarchyTable/Source/Editor/Private/FloatColumn.cpp
@@ -25,7 +25,9 @@ TSharedRef<SWidget> FHierarchyTableColumn_Float::CreateEntryWidget(TObjectPtr<UH
 			})
 		.OnValueChanged_Lambda([HierarchyTable, EntryIndex](float NewValue)
 			{
+				// TODO: Update API to avoid having to manually regenerate guid
 				HierarchyTable->GetMutableTableEntry(EntryIndex)->GetMutableValue<FHierarchyTable_ElementType_Float>()->Value = NewValue;
+				HierarchyTable->RegenerateEntriesGuid();
 			})
 		.OnBeginSliderMovement_Lambda([HierarchyTable]()
 			{
diff --git a/Engine/Plugins/Experimental/Animation/HierarchyTableAnimation/Source/Runtime/Private/HierarchyTableBlendProfile.cpp b/Engine/Plugins/Experimental/Animation/HierarchyTableAnimation/Source/Runtime/Private/HierarchyTableBlendProfile.cpp
--- a/Engine/Plugins/Experimental/Animation/HierarchyTableAnimation/Source/Runtime/Private/HierarchyTableBlendProfile.cpp
+++ b/Engine/Plugins/Experimental/Animation/HierarchyTableAnimation/Source/Runtime/Private/HierarchyTableBlendProfile.cpp
@@ -208,7 +208,7 @@ FHierarchyTableBlendProfile::FHierarchyTableBlendProfile(TObjectPtr<UHierarchyTa
 					UE::Anim::FAttributeId Attribute(Entry.Identifier, Entry.Parent, TEXT("bone"));
 
 					const float EntryWeight = Entry.GetValue<FHierarchyTable_ElementType_Mask>()->Value;
-					AttributeBlendWeights.Add(FMaskedAttributeWeight(Attribute, EntryWeight));
+					AttributeBlendWeights.Add(FBlendProfileStandaloneCachedData::FMaskedAttributeWeight(Attribute, EntryWeight));
 				}
 			}
 		}
@@ -269,7 +269,7 @@ const UE::Anim::TNamedValueArray<FDefaultAllocator, UE::Anim::FCurveElement>& FH
 	return CurveBlendWeights;
 }
 
-const TArray<FHierarchyTableBlendProfile::FMaskedAttributeWeight>& FHierarchyTableBlendProfile::GetAttributeBlendWeights() const
+const TArray<FBlendProfileStandaloneCachedData::FMaskedAttributeWeight>& FHierarchyTableBlendProfile::GetAttributeBlendWeights() const
 {
 	return AttributeBlendWeights;
 }
diff --git a/Engine/Plugins/Experimental/Animation/HierarchyTableAnimation/Source/Runtime/Public/HierarchyTableBlendProfile.h b/Engine/Plugins/Experimental/Animation/HierarchyTableAnimation/Source/Runtime/Public/HierarchyTableBlendProfile.h
--- a/Engine/Plugins/Experimental/Animation/HierarchyTableAnimation/Source/Runtime/Public/HierarchyTableBlendProfile.h
+++ b/Engine/Plugins/Experimental/Animation/HierarchyTableAnimation/Source/Runtime/Public/HierarchyTableBlendProfile.h
@@ -128,20 +128,8 @@ public:
 	UE_API virtual EBlendProfileMode GetMode() const override;
 	UE_API virtual TObjectPtr<USkeleton> GetSkeleton() const override;
 
-	struct FMaskedAttributeWeight
-	{
-		UE::Anim::FAttributeId Attribute;
-		float Weight;
-
-		FMaskedAttributeWeight(const UE::Anim::FAttributeId InAttribute, const float InWeight)
-			: Attribute(InAttribute)
-			, Weight(InWeight)
-		{
-		}
-	};
-
 	UE_API const UE::Anim::TNamedValueArray<FDefaultAllocator, UE::Anim::FCurveElement>& GetCurveBlendWeights() const;
-	UE_API const TArray<FMaskedAttributeWeight>& GetAttributeBlendWeights() const;
+	UE_API const TArray<FBlendProfileStandaloneCachedData::FMaskedAttributeWeight>& GetAttributeBlendWeights() const;
 
 	UE_API void ConstructBlendProfile(const TObjectPtr<UBlendProfile> OutBlendProfile) const;
 
@@ -154,7 +142,7 @@ private:
 
 	UE::Anim::TNamedValueArray<FDefaultAllocator, UE::Anim::FCurveElement> CurveBlendWeights;
 
-	TArray<FMaskedAttributeWeight> AttributeBlendWeights;
+	TArray<FBlendProfileStandaloneCachedData::FMaskedAttributeWeight> AttributeBlendWeights;
 
 	EBlendProfileMode Mode = EBlendProfileMode::WeightFactor;
 };
diff --git a/Engine/Plugins/Experimental/UAF/UAF/Source/UAF/Private/AnimNextModuleImpl.cpp b/Engine/Plugins/Experimental/UAF/UAF/Source/UAF/Private/AnimNextModuleImpl.cpp
--- a/Engine/Plugins/Experimental/UAF/UAF/Source/UAF/Private/AnimNextModuleImpl.cpp
+++ b/Engine/Plugins/Experimental/UAF/UAF/Source/UAF/Private/AnimNextModuleImpl.cpp
@@ -39,6 +39,8 @@
 #include "UAF/UAFAssetFactory.h"
 #include "UAF/ValueRuntime/ValueBundle.h"
 #include "UAF/Attributes/EngineAttributes.h"
+#include "UAF/BlendMask/UAFBlendMask.h"
+#include "UAF/BlendProfile/UAFBlendProfile.h"
 
 #define LOCTEXT_NAMESPACE "AnimNextModule"
 
@@ -63,12 +65,14 @@ namespace UE::UAF
 			{ UBlendSpace::StaticClass(), FRigVMRegistry::ERegisterObjectOperation::Class },
 			{ UBlendSpace1D::StaticClass(), FRigVMRegistry::ERegisterObjectOperation::Class },
 			{ UScriptStruct::StaticClass(), FRigVMRegistry::ERegisterObjectOperation::Class },
-			{ UBlendProfile::StaticClass(), FRigVMRegistry::ERegisterObjectOperation::Class },
+			{ UBlendProfile::StaticClass(), FRigVMRegistry::ERegisterObjectOperation::Class }, // TODO: Remove
 			{ UCurveFloat::StaticClass(), FRigVMRegistry::ERegisterObjectOperation::Class },
 			{ UUAFComponent::StaticClass(), FRigVMRegistry::ERegisterObjectOperation::Class },
 			{ UUAFSystem::StaticClass(), FRigVMRegistry::ERegisterObjectOperation::Class },
 			{ UUAFRigVMAsset::StaticClass(), FRigVMRegistry::ERegisterObjectOperation::Class },
 			{ UHierarchyTable::StaticClass(), FRigVMRegistry::ERegisterObjectOperation::Class },
+			{ UUAFBlendMask::StaticClass(), FRigVMRegistry::ERegisterObjectOperation::Class },
+			{ UUAFBlendProfile::StaticClass(), FRigVMRegistry::ERegisterObjectOperation::Class },
 			{ UCharacterMovementComponent::StaticClass(), FRigVMRegistry::ERegisterObjectOperation::Class },
 		};
 
diff --git a/Engine/Plugins/Experimental/UAF/UAF/Source/UAF/Private/UAF/BlendMask/UAFBlendMask.cpp b/Engine/Plugins/Experimental/UAF/UAF/Source/UAF/Private/UAF/BlendMask/UAFBlendMask.cpp
new file mode 100644
--- /dev/null
+++ b/Engine/Plugins/Experimental/UAF/UAF/Source/UAF/Private/UAF/BlendMask/UAFBlendMask.cpp
@@ -0,0 +1,253 @@
+// Copyright Epic Games, Inc. All Rights Reserved.
+
+#include "UAF/BlendMask/UAFBlendMask.h"
+
+#if WITH_EDITOR
+#include "Animation/Skeleton.h"
+#include "SkeletonHierarchyTableType.h"
+#include "UObject/AssetRegistryTagsContext.h"
+#include "UObject/NoExportTypes.h"
+#include "UObject/ObjectSaveContext.h"
+#include "StructUtils/InstancedStruct.h"
+#include "Misc/Optional.h"
+#endif // WITH_EDITOR
+
+#include UE_INLINE_GENERATED_CPP_BY_NAME(UAFBlendMask)
+
+#if WITH_EDITOR
+
+void UUAFBlendMask::BeginDestroy()
+{
+	if (Skeleton && !IsTemplate())
+	{
+		Skeleton->UnregisterOnSkeletonHierarchyChanged(this);
+	}
+
+	Super::BeginDestroy();
+}
+
+void UUAFBlendMask::HandleSkeletonHierarchyChanged()
+{
+	UpdateHierarchy();
+	UpdateCachedData();
+}
+
+void UUAFBlendMask::UpdateHierarchy()
+{
+	/* Skeleton hierarchy tables take a snapshot of the currently bound skeleton's hierarchy.
+	 *  When the skeleton hierarchy changes and this function is called we do the following:
+	 *		1. Make a note of all the bone table entries with overridden values.
+	 *		2. Make a note of all the curves and attributes in the table and their location in the hierarchy.
+	 *		3. Clear the hierarchy table contents
+	 *		4. Reconstruct the bone hierarchy using the new skeleton hierarchy and insert the old blend values 
+	 *		   for the old overridden bones.
+	 *		5. Reconstruct the curves and attributes hierarchy.
+	 */
+
+	if (!Table)
+	{
+		return;
+	}
+
+	const FHierarchyTable_TableType_Skeleton& Metadata = Table->GetTableMetadata<FHierarchyTable_TableType_Skeleton>();
+	SetSkeleton(Metadata.Skeleton);
+
+	const FReferenceSkeleton& RefSkeleton = GetSkeleton()->GetReferenceSkeleton();
+	const int32 BoneCount = RefSkeleton.GetNum();
+
+	struct FCurveAttributeBlendData
+	{
+		bool bIsCurve;
+		FName Identifier;
+		FName ParentIdentifier;
+		TOptional<FInstancedStruct> Payload;
+	};
+
+	TMap<FName, FInstancedStruct> BoneData;
+	TArray<FCurveAttributeBlendData> CurveAndAttributeData;
+	{
+		for (const FHierarchyTableEntryData& TableEntry : Table->GetTableData())
+		{
+			const auto& EntryMetadata = TableEntry.GetMetadata<FHierarchyTable_TablePayloadType_Skeleton>();
+			if (EntryMetadata.EntryType == ESkeletonHierarchyTable_TablePayloadEntryType::Bone)
+			{
+				if (TableEntry.IsOverridden())
+				{
+					BoneData.Add(TableEntry.Identifier, TableEntry.GetPayload().GetValue());
+				}
+			}
+			if (EntryMetadata.EntryType == ESkeletonHierarchyTable_TablePayloadEntryType::Curve)
+			{
+				check(TableEntry.HasParent()); // Curves must be parented
+				const int32 ParentIndex = TableEntry.Parent;
+
+				const FHierarchyTableEntryData* ParentEntry = Table->GetTableEntry(ParentIndex);
+				check(ParentEntry);
+
+				CurveAndAttributeData.Add({ true, TableEntry.Identifier, ParentEntry->Identifier, TableEntry.GetPayload() });
+			}
+			else if (EntryMetadata.EntryType == ESkeletonHierarchyTable_TablePayloadEntryType::Attribute)
+			{
+				check(TableEntry.HasParent()); // Attributes must be parented
+				const int32 ParentIndex = TableEntry.Parent;
+
+				const FHierarchyTableEntryData* ParentEntry = Table->GetTableEntry(ParentIndex);
+				check(ParentEntry);
+
+				CurveAndAttributeData.Add({ false, TableEntry.Identifier, ParentEntry->Identifier, TableEntry.GetPayload() });
+			}
+		}
+	}
+
+	Table->EmptyTable();
+
+	// Re-add bone data
+	{
+		FInstancedStruct DefaultTablePayload;
+		DefaultTablePayload.InitializeAs<FHierarchyTable_TablePayloadType_Skeleton>();
+		DefaultTablePayload.GetMutable<FHierarchyTable_TablePayloadType_Skeleton>().EntryType = ESkeletonHierarchyTable_TablePayloadEntryType::Bone;
+
+		for (int32 BoneIndex = 0; BoneIndex < BoneCount; ++BoneIndex)
+		{
+			const FName BoneName = RefSkeleton.GetBoneName(BoneIndex);
+
+			FHierarchyTableEntryData EntryData;
+			EntryData.Parent = RefSkeleton.GetParentIndex(BoneIndex);
+			EntryData.Identifier = BoneName;
+			EntryData.TablePayload = DefaultTablePayload;
+			EntryData.OwnerTable = Table;
+
+			if (BoneData.Contains(BoneName))
+			{
+				EntryData.Payload = BoneData[BoneName];
+			}
+			else
+			{
+				EntryData.Payload = (BoneIndex == 0) ? Table->CreateDefaultValue() : TOptional<FInstancedStruct>();
+			}
+
+			Table->AddEntry(EntryData);
+		}
+	}
+
+	// Re-add curve and attribute data
+	{
+		FInstancedStruct DefaultTablePayload_Curve;
+		DefaultTablePayload_Curve.InitializeAs<FHierarchyTable_TablePayloadType_Skeleton>();
+		DefaultTablePayload_Curve.GetMutable<FHierarchyTable_TablePayloadType_Skeleton>().EntryType = ESkeletonHierarchyTable_TablePayloadEntryType::Curve;
+
+		FInstancedStruct DefaultTablePayload_Attribute;
+		DefaultTablePayload_Attribute.InitializeAs<FHierarchyTable_TablePayloadType_Skeleton>();
+		DefaultTablePayload_Attribute.GetMutable<FHierarchyTable_TablePayloadType_Skeleton>().EntryType = ESkeletonHierarchyTable_TablePayloadEntryType::Attribute;
+
+		for (const FCurveAttributeBlendData& Entry : CurveAndAttributeData)
+		{
+			const FHierarchyTableEntryData* const ParentEntry = Table->GetTableEntry(Entry.ParentIdentifier);
+			if (ParentEntry)
+			{
+				FHierarchyTableEntryData EntryData;
+				EntryData.Identifier = Entry.Identifier;
+				EntryData.Parent = Table->GetTableEntryIndex(ParentEntry->Identifier);
+				EntryData.TablePayload = Entry.bIsCurve ? DefaultTablePayload_Curve : DefaultTablePayload_Attribute;
+				EntryData.OwnerTable = Table;
+				EntryData.Payload = Entry.Payload;
+
+				Table->AddEntry(EntryData);
+			}
+			else
+			{
+				// This curve/attribute was attached to a parent that no longer exists
+			}
+		}
+	}
+
+	Table->RegenerateEntriesGuid();
+}
+
+void UUAFBlendMask::UpdateCachedData()
+{
+	CachedBlendProfileData.Reset();
+
+	if (Table)
+	{
+		CachedBlendProfileData.Init(Table, EBlendProfileMode::BlendMask);
+
+		const FHierarchyTable_TableType_Skeleton& Metadata = Table->GetTableMetadata<FHierarchyTable_TableType_Skeleton>();
+		SetSkeleton(Metadata.Skeleton);
+	}
+}
+
+#endif // WITH_EDITOR
+
+TObjectPtr<USkeleton> UUAFBlendMask::GetSkeleton() const
+{
+	return Skeleton;
+}
+
+void UUAFBlendMask::PostLoad()
+{
+	Super::PostLoad();
+
+#if WITH_EDITOR
+	// Update the hierarchy to reflect any changes made to the skeleton
+	UpdateHierarchy();
+	// Cache the flattened data for runtime use
+	UpdateCachedData();
+#endif // WITH_EDITOR
+
+	CachedBlendProfileData.UnpackCachedData();
+}
+
+#if WITH_EDITOR
+
+void UUAFBlendMask::PreSave(FObjectPreSaveContext ObjectSaveContext)
+{
+	UpdateCachedData();
+	
+	Super::PreSave(ObjectSaveContext);
+}
+
+void UUAFBlendMask::GetAssetRegistryTags(FAssetRegistryTagsContext Context) const
+{
+	Super::GetAssetRegistryTags(Context);
+
+	if (Table)
+	{
+		const FHierarchyTable_TableType_Skeleton& TableMetadata = Table->GetTableMetadata<FHierarchyTable_TableType_Skeleton>();
+
+		Context.AddTag( FAssetRegistryTag("Skeleton", FAssetData(TableMetadata.Skeleton).ToSoftObjectPath().ToString(), FAssetRegistryTag::TT_Alphabetical) );
+	}
+}
+
+void UUAFBlendMask::SetSkeleton(TObjectPtr<USkeleton> InSkeleton)
+{
+	if (Skeleton)
+	{
+		Skeleton->UnregisterOnSkeletonHierarchyChanged(this);
+	}
+	
+	if (ensure(InSkeleton))
+	{
+		Skeleton = InSkeleton;
+
+		Skeleton->RegisterOnSkeletonHierarchyChanged(
+			USkeleton::FOnSkeletonHierarchyChanged::CreateUObject(this, &UUAFBlendMask::HandleSkeletonHierarchyChanged));
+	}
+}
+
+const UUAFBlendMask::FSkeletonBoneWeightArray& UUAFBlendMask::GetSkeletonBoneWeights() const
+{
+	return CachedBlendProfileData.GetBoneBlendWeights();
+}
+
+const UUAFBlendMask::FSkeletonCurveWeightArray& UUAFBlendMask::GetSkeletonCurveWeights() const
+{
+	return CachedBlendProfileData.GetCurveBlendWeights();
+}
+
+const UUAFBlendMask::FSkeletonAttributeWeightArray& UUAFBlendMask::GetSkeletonAttributeWeights() const
+{
+	return CachedBlendProfileData.GetAttributeBlendWeights();
+}
+
+#endif // WITH_EDITOR
diff --git a/Engine/Plugins/Experimental/UAF/UAF/Source/UAF/Private/UAF/BlendProfile/UAFBlendProfile.cpp b/Engine/Plugins/Experimental/UAF/UAF/Source/UAF/Private/UAF/BlendProfile/UAFBlendProfile.cpp
new file mode 100644
--- /dev/null
+++ b/Engine/Plugins/Experimental/UAF/UAF/Source/UAF/Private/UAF/BlendProfile/UAFBlendProfile.cpp
@@ -0,0 +1,274 @@
+// Copyright Epic Games, Inc. All Rights Reserved.
+
+#include "UAF/BlendProfile/UAFBlendProfile.h"
+
+#if WITH_EDITOR
+#include "Animation/Skeleton.h"
+#include "SkeletonHierarchyTableType.h"
+#include "UObject/AssetRegistryTagsContext.h"
+#include "UObject/NoExportTypes.h"
+#include "UObject/ObjectSaveContext.h"
+#include "StructUtils/InstancedStruct.h"
+#include "Misc/Optional.h"
+#endif // WITH_EDITOR
+
+#include UE_INLINE_GENERATED_CPP_BY_NAME(UAFBlendProfile)
+
+#if WITH_EDITOR
+
+void UUAFBlendProfile::BeginDestroy()
+{
+	if (Skeleton && !IsTemplate())
+	{
+		Skeleton->UnregisterOnSkeletonHierarchyChanged(this);
+	}
+
+	Super::BeginDestroy();
+}
+
+void UUAFBlendProfile::HandleSkeletonHierarchyChanged()
+{
+	UpdateHierarchy();
+	UpdateCachedData();
+}
+
+void UUAFBlendProfile::UpdateHierarchy()
+{
+	/* Skeleton hierarchy tables take a snapshot of the currently bound skeleton's hierarchy.
+	 *  When the skeleton hierarchy changes and this function is called we do the following:
+	 *		1. Make a note of all the bone table entries with overridden values.
+	 *		2. Make a note of all the curves and attributes in the table and their location in the hierarchy.
+	 *		3. Clear the hierarchy table contents
+	 *		4. Reconstruct the bone hierarchy using the new skeleton hierarchy and insert the old blend values 
+	 *		   for the old overridden bones.
+	 *		5. Reconstruct the curves and attributes hierarchy.
+	 */
+
+	if (!Table)
+	{
+		return;
+	}
+
+	const FHierarchyTable_TableType_Skeleton& Metadata = Table->GetTableMetadata<FHierarchyTable_TableType_Skeleton>();
+	SetSkeleton(Metadata.Skeleton);
+
+	const FReferenceSkeleton& RefSkeleton = GetSkeleton()->GetReferenceSkeleton();
+	const int32 BoneCount = RefSkeleton.GetNum();
+
+	struct FCurveAttributeBlendData
+	{
+		bool bIsCurve;
+		FName Identifier;
+		FName ParentIdentifier;
+		TOptional<FInstancedStruct> Payload;
+	};
+
+	TMap<FName, FInstancedStruct> BoneData;
+	TArray<FCurveAttributeBlendData> CurveAndAttributeData;
+	{
+		for (const FHierarchyTableEntryData& TableEntry : Table->GetTableData())
+		{
+			const auto& EntryMetadata = TableEntry.GetMetadata<FHierarchyTable_TablePayloadType_Skeleton>();
+			if (EntryMetadata.EntryType == ESkeletonHierarchyTable_TablePayloadEntryType::Bone)
+			{
+				if (TableEntry.IsOverridden())
+				{
+					BoneData.Add(TableEntry.Identifier, TableEntry.GetPayload().GetValue());
+				}
+			}
+			if (EntryMetadata.EntryType == ESkeletonHierarchyTable_TablePayloadEntryType::Curve)
+			{
+				check(TableEntry.HasParent()); // Curves must be parented
+				const int32 ParentIndex = TableEntry.Parent;
+
+				const FHierarchyTableEntryData* ParentEntry = Table->GetTableEntry(ParentIndex);
+				check(ParentEntry);
+
+				CurveAndAttributeData.Add({ true, TableEntry.Identifier, ParentEntry->Identifier, TableEntry.GetPayload() });
+			}
+			else if (EntryMetadata.EntryType == ESkeletonHierarchyTable_TablePayloadEntryType::Attribute)
+			{
+				check(TableEntry.HasParent()); // Attributes must be parented
+				const int32 ParentIndex = TableEntry.Parent;
+
+				const FHierarchyTableEntryData* ParentEntry = Table->GetTableEntry(ParentIndex);
+				check(ParentEntry);
+
+				CurveAndAttributeData.Add({ false, TableEntry.Identifier, ParentEntry->Identifier, TableEntry.GetPayload() });
+			}
+		}
+	}
+
+	Table->EmptyTable();
+
+	// Re-add bone data
+	{
+		FInstancedStruct DefaultTablePayload;
+		DefaultTablePayload.InitializeAs<FHierarchyTable_TablePayloadType_Skeleton>();
+		DefaultTablePayload.GetMutable<FHierarchyTable_TablePayloadType_Skeleton>().EntryType = ESkeletonHierarchyTable_TablePayloadEntryType::Bone;
+
+		for (int32 BoneIndex = 0; BoneIndex < BoneCount; ++BoneIndex)
+		{
+			const FName BoneName = RefSkeleton.GetBoneName(BoneIndex);
+
+			FHierarchyTableEntryData EntryData;
+			EntryData.Parent = RefSkeleton.GetParentIndex(BoneIndex);
+			EntryData.Identifier = BoneName;
+			EntryData.TablePayload = DefaultTablePayload;
+			EntryData.OwnerTable = Table;
+
+			if (BoneData.Contains(BoneName))
+			{
+				EntryData.Payload = BoneData[BoneName];
+			}
+			else
+			{
+				EntryData.Payload = (BoneIndex == 0) ? Table->CreateDefaultValue() : TOptional<FInstancedStruct>();
+			}
+
+			Table->AddEntry(EntryData);
+		}
+	}
+
+	// Re-add curve and attribute data
+	{
+		FInstancedStruct DefaultTablePayload_Curve;
+		DefaultTablePayload_Curve.InitializeAs<FHierarchyTable_TablePayloadType_Skeleton>();
+		DefaultTablePayload_Curve.GetMutable<FHierarchyTable_TablePayloadType_Skeleton>().EntryType = ESkeletonHierarchyTable_TablePayloadEntryType::Curve;
+
+		FInstancedStruct DefaultTablePayload_Attribute;
+		DefaultTablePayload_Attribute.InitializeAs<FHierarchyTable_TablePayloadType_Skeleton>();
+		DefaultTablePayload_Attribute.GetMutable<FHierarchyTable_TablePayloadType_Skeleton>().EntryType = ESkeletonHierarchyTable_TablePayloadEntryType::Attribute;
+
+		for (const FCurveAttributeBlendData& Entry : CurveAndAttributeData)
+		{
+			const FHierarchyTableEntryData* const ParentEntry = Table->GetTableEntry(Entry.ParentIdentifier);
+			if (ParentEntry)
+			{
+				FHierarchyTableEntryData EntryData;
+				EntryData.Identifier = Entry.Identifier;
+				EntryData.Parent = Table->GetTableEntryIndex(ParentEntry->Identifier);
+				EntryData.TablePayload = Entry.bIsCurve ? DefaultTablePayload_Curve : DefaultTablePayload_Attribute;
+				EntryData.OwnerTable = Table;
+				EntryData.Payload = Entry.Payload;
+
+				Table->AddEntry(EntryData);
+			}
+			else
+			{
+				// This curve/attribute was attached to a parent that no longer exists
+			}
+		}
+	}
+
+	Table->RegenerateEntriesGuid();
+}
+
+void UUAFBlendProfile::UpdateCachedData()
+{
+	CachedBlendProfileData.Reset();
+
+	if (Table)
+	{
+		CachedBlendProfileData.Init(Table, Type == EUAFBlendProfileType::WeightFactor ? EBlendProfileMode::WeightFactor : EBlendProfileMode::TimeFactor);
+
+		const FHierarchyTable_TableType_Skeleton& Metadata = Table->GetTableMetadata<FHierarchyTable_TableType_Skeleton>();
+		SetSkeleton(Metadata.Skeleton);
+	}
+}
+
+#endif // WITH_EDITOR
+
+TObjectPtr<USkeleton> UUAFBlendProfile::GetSkeleton() const
+{
+	return Skeleton;
+}
+
+EUAFBlendProfileType UUAFBlendProfile::GetType() const
+{
+	return Type;
+}
+
+EBlendProfileMode UUAFBlendProfile::GetEngineType() const
+{
+	return Type == EUAFBlendProfileType::TimeFactor
+		? EBlendProfileMode::TimeFactor
+		: EBlendProfileMode::WeightFactor;
+}
+
+void UUAFBlendProfile::PostLoad()
+{
+	Super::PostLoad();
+
+#if WITH_EDITOR
+	// Update the hierarchy to reflect any changes made to the skeleton
+	UpdateHierarchy();
+	// Cache the flattened data for runtime use
+	UpdateCachedData();
+#endif // WITH_EDITOR
+
+	CachedBlendProfileData.UnpackCachedData();
+}
+
+#if WITH_EDITOR
+
+void UUAFBlendProfile::PreSave(FObjectPreSaveContext ObjectSaveContext)
+{
+	UpdateCachedData();
+	
+	Super::PreSave(ObjectSaveContext);
+}
+
+void UUAFBlendProfile::GetAssetRegistryTags(FAssetRegistryTagsContext Context) const
+{
+	Super::GetAssetRegistryTags(Context);
+
+	if (Table)
+	{
+		const FHierarchyTable_TableType_Skeleton& TableMetadata = Table->GetTableMetadata<FHierarchyTable_TableType_Skeleton>();
+
+		Context.AddTag( FAssetRegistryTag("Skeleton", FAssetData(TableMetadata.Skeleton).ToSoftObjectPath().ToString(), FAssetRegistryTag::TT_Alphabetical) );
+	}
+
+	using namespace UE::UAF;
+	
+	Context.AddTag(FAssetRegistryTag(
+		AssetRegistryTag_BlendProfileType,
+		Type == EUAFBlendProfileType::TimeFactor
+			? AssetRegistryTag_BlendProfileType_TimeProfile
+			: AssetRegistryTag_BlendProfileType_WeightProfile,
+		FAssetRegistryTag::TT_Alphabetical));
+}
+
+void UUAFBlendProfile::SetSkeleton(TObjectPtr<USkeleton> InSkeleton)
+{
+	if (Skeleton)
+	{
+		Skeleton->UnregisterOnSkeletonHierarchyChanged(this);
+	}
+	
+	if (ensure(InSkeleton))
+	{
+		Skeleton = InSkeleton;
+
+		Skeleton->RegisterOnSkeletonHierarchyChanged(
+			USkeleton::FOnSkeletonHierarchyChanged::CreateUObject(this, &UUAFBlendProfile::HandleSkeletonHierarchyChanged));
+	}
+}
+
+const UUAFBlendProfile::FSkeletonBoneWeightArray& UUAFBlendProfile::GetSkeletonBoneWeights() const
+{
+	return CachedBlendProfileData.GetBoneBlendWeights();
+}
+
+const UUAFBlendProfile::FSkeletonCurveWeightArray& UUAFBlendProfile::GetSkeletonCurveWeights() const
+{
+	return CachedBlendProfileData.GetCurveBlendWeights();
+}
+
+const UUAFBlendProfile::FSkeletonAttributeWeightArray& UUAFBlendProfile::GetSkeletonAttributeWeights() const
+{
+	return CachedBlendProfileData.GetAttributeBlendWeights();
+}
+
+#endif // WITH_EDITOR
diff --git a/Engine/Plugins/Experimental/UAF/UAF/Source/UAF/Public/UAF/BlendMask/UAFBlendMask.h b/Engine/Plugins/Experimental/UAF/UAF/Source/UAF/Public/UAF/BlendMask/UAFBlendMask.h
new file mode 100644
--- /dev/null
+++ b/Engine/Plugins/Experimental/UAF/UAF/Source/UAF/Public/UAF/BlendMask/UAFBlendMask.h
@@ -0,0 +1,76 @@
+// Copyright Epic Games, Inc. All Rights Reserved.
+
+#pragma once
+
+#include "CoreMinimal.h"
+#include "HierarchyTable.h"
+#include "HierarchyTableBlendProfile.h"
+
+#include "UAFBlendMask.generated.h"
+
+#define UE_API UAF_API
+
+class UUAFBlendMaskFactory;
+namespace UE::UAF
+{
+class FBlendMaskEditorToolkit;
+}
+
+// An asset to define what parts to include in a blend when layering.
+UCLASS(MinimalAPI, EditInlineNew, BlueprintType)
+class UUAFBlendMask : public UObject
+{
+	GENERATED_BODY()
+
+public:
+	using FSkeletonBoneWeightArray = TArray<float>;
+	using FSkeletonCurveWeightArray = UE::Anim::TNamedValueArray<FDefaultAllocator, UE::Anim::FCurveElement>;
+	using FSkeletonAttributeWeightArray = TArray<FBlendProfileStandaloneCachedData::FMaskedAttributeWeight>;
+
+	// Begin UObject
+	UE_API virtual void PostLoad() override;
+#if WITH_EDITOR
+	UE_API virtual void BeginDestroy() override;
+	UE_API virtual void PreSave(FObjectPreSaveContext ObjectSaveContext) override;
+	UE_API virtual void GetAssetRegistryTags(FAssetRegistryTagsContext Context) const override;
+#endif
+	// End UObject
+
+#if WITH_EDITOR
+	UE_API void UpdateHierarchy();
+
+	UE_API void UpdateCachedData();
+#endif
+
+	UE_API TObjectPtr<USkeleton> GetSkeleton() const;
+
+	UE_API const FSkeletonBoneWeightArray& GetSkeletonBoneWeights() const;
+	UE_API const FSkeletonCurveWeightArray& GetSkeletonCurveWeights() const;
+	UE_API const FSkeletonAttributeWeightArray& GetSkeletonAttributeWeights() const;
+
+private:
+#if WITH_EDITORONLY_DATA
+	UPROPERTY()
+	TObjectPtr<UHierarchyTable> Table;
+#endif // WITH_EDITORONLY_DATA
+
+	// The flattened version of the table data for use at runtime instead of,
+	// slowly traversing the table tree.
+	UPROPERTY()
+	FBlendProfileStandaloneCachedData CachedBlendProfileData;
+
+private:
+#if WITH_EDITOR
+	void SetSkeleton(TObjectPtr<USkeleton> InSkeleton);
+
+	void HandleSkeletonHierarchyChanged();
+#endif // WITH_EDITOR
+
+	UPROPERTY(VisibleAnywhere, Category = Default)
+	TObjectPtr<USkeleton> Skeleton;
+
+	friend class UUAFBlendMaskFactory;
+	friend class UE::UAF::FBlendMaskEditorToolkit;
+};
+
+#undef UE_API
\ No newline at end of file
diff --git a/Engine/Plugins/Experimental/UAF/UAF/Source/UAF/Public/UAF/BlendProfile/UAFBlendProfile.h b/Engine/Plugins/Experimental/UAF/UAF/Source/UAF/Public/UAF/BlendProfile/UAFBlendProfile.h
new file mode 100644
--- /dev/null
+++ b/Engine/Plugins/Experimental/UAF/UAF/Source/UAF/Public/UAF/BlendProfile/UAFBlendProfile.h
@@ -0,0 +1,98 @@
+// Copyright Epic Games, Inc. All Rights Reserved.
+
+#pragma once
+
+#include "CoreMinimal.h"
+#include "HierarchyTable.h"
+#include "HierarchyTableBlendProfile.h"
+
+#include "UAFBlendProfile.generated.h"
+
+#define UE_API UAF_API
+
+class UUAFBlendProfileFactory;
+namespace UE::UAF
+{
+	class FBlendProfileEditorToolkit;
+
+	inline const FName AssetRegistryTag_BlendProfileType = TEXT("BlendProfileType");
+	inline const FString AssetRegistryTag_BlendProfileType_TimeProfile = TEXT("Time");
+	inline const FString AssetRegistryTag_BlendProfileType_WeightProfile = TEXT("Weight");
+}
+
+UENUM()
+enum class EUAFBlendProfileType
+{
+	WeightFactor,
+	TimeFactor,
+};
+
+// An asset to define how a blend should be applied in a transition with per-bone controls.
+UCLASS(MinimalAPI, EditInlineNew, BlueprintType)
+class UUAFBlendProfile : public UObject
+{
+	GENERATED_BODY()
+
+public:
+	using FSkeletonBoneWeightArray = TArray<float>;
+	using FSkeletonCurveWeightArray = UE::Anim::TNamedValueArray<FDefaultAllocator, UE::Anim::FCurveElement>;
+	using FSkeletonAttributeWeightArray = TArray<FBlendProfileStandaloneCachedData::FMaskedAttributeWeight>;
+	
+	// Begin UObject
+	UE_API virtual void PostLoad() override;
+#if WITH_EDITOR
+	UE_API virtual void BeginDestroy() override;
+	UE_API virtual void PreSave(FObjectPreSaveContext ObjectSaveContext) override;
+	UE_API virtual void GetAssetRegistryTags(FAssetRegistryTagsContext Context) const override;
+#endif
+	// End UObject
+
+#if WITH_EDITOR
+	UE_API void UpdateHierarchy();
+
+	UE_API void UpdateCachedData();
+#endif
+
+	UE_API TObjectPtr<USkeleton> GetSkeleton() const;
+
+	// Returns whether this blend profile is a time-based or weight-based time profile.
+	UE_API EUAFBlendProfileType GetType() const;
+
+	// Similar to GetType but returns the value in terms of the engine blend profile mode.
+	// Used for reusing UBlendProfile engine code for UAF blend profiles.
+	// Will never return EBlendProfileMode::BlendMask.
+	UE_API EBlendProfileMode GetEngineType() const;
+
+	UE_API const FSkeletonBoneWeightArray& GetSkeletonBoneWeights() const;
+	UE_API const FSkeletonCurveWeightArray& GetSkeletonCurveWeights() const;
+	UE_API const FSkeletonAttributeWeightArray& GetSkeletonAttributeWeights() const;
+
+private:
+	UPROPERTY(VisibleAnywhere, Category = Default, AssetRegistrySearchable)
+	EUAFBlendProfileType Type;
+
+#if WITH_EDITORONLY_DATA
+	UPROPERTY()
+	TObjectPtr<UHierarchyTable> Table;
+#endif // WITH_EDITORONLY_DATA
+
+	// The flattened version of the table data for use at runtime instead of,
+	// slowly traversing the table tree.
+	UPROPERTY()
+	FBlendProfileStandaloneCachedData CachedBlendProfileData;
+
+private:
+#if WITH_EDITOR
+	void SetSkeleton(TObjectPtr<USkeleton> InSkeleton);
+
+	void HandleSkeletonHierarchyChanged();
+#endif // WITH_EDITOR
+
+	UPROPERTY(VisibleAnywhere, Category = Default)
+	TObjectPtr<USkeleton> Skeleton;
+
+	friend class UUAFBlendProfileFactory;
+	friend class UE::UAF::FBlendProfileEditorToolkit;
+};
+
+#undef UE_API
\ No newline at end of file
diff --git a/Engine/Plugins/Experimental/UAF/UAF/Source/UAF/UAF.Build.cs b/Engine/Plugins/Experimental/UAF/UAF/Source/UAF/UAF.Build.cs
--- a/Engine/Plugins/Experimental/UAF/UAF/Source/UAF/UAF.Build.cs
+++ b/Engine/Plugins/Experimental/UAF/UAF/Source/UAF/UAF.Build.cs
@@ -15,6 +15,7 @@ namespace UnrealBuildTool.Rules
 					"Engine",
 					"RenderCore",
 					"HierarchyTableRuntime",
+					"HierarchyTableAnimationRuntime",
 					"RewindDebuggerRuntimeInterface"
 				}
 			);
diff --git a/Engine/Plugins/Experimental/UAF/UAF/Source/UAFEditor/Internal/UAFStyle.h b/Engine/Plugins/Experimental/UAF/UAF/Source/UAFEditor/Internal/UAFStyle.h
--- a/Engine/Plugins/Experimental/UAF/UAF/Source/UAFEditor/Internal/UAFStyle.h
+++ b/Engine/Plugins/Experimental/UAF/UAF/Source/UAFEditor/Internal/UAFStyle.h
@@ -27,6 +27,8 @@ private:
 	UE_API void SetAssetWizardStyles();
 
 	UE_API void SetSetBindingStyles();
+	
+	UE_API void SetBlendProfileStyles();
 };
 
 #undef UE_API
diff --git a/Engine/Plugins/Experimental/UAF/UAF/Source/UAFEditor/Private/UAF/BlendMask/BlendMaskAssetDefinition.cpp b/Engine/Plugins/Experimental/UAF/UAF/Source/UAFEditor/Private/UAF/BlendMask/BlendMaskAssetDefinition.cpp
new file mode 100644
--- /dev/null
+++ b/Engine/Plugins/Experimental/UAF/UAF/Source/UAFEditor/Private/UAF/BlendMask/BlendMaskAssetDefinition.cpp
@@ -0,0 +1,44 @@
+// Copyright Epic Games, Inc. All Rights Reserved.
+
+#include "BlendMaskAssetDefinition.h"
+
+#include "UAF/BlendMask/UAFBlendMask.h"
+#include "UAF/BlendMask/BlendMaskEditorToolkit.h"
+
+#include UE_INLINE_GENERATED_CPP_BY_NAME(BlendMaskAssetDefinition)
+
+#define LOCTEXT_NAMESPACE "UAFBlendMask"
+
+FText UAssetDefinition_UAFBlendMask::GetAssetDisplayName() const
+{
+	return LOCTEXT("AssetDisplayName", "UAF Blend Mask");
+}
+
+FLinearColor UAssetDefinition_UAFBlendMask::GetAssetColor() const
+{
+	return FLinearColor(FColor::Yellow);
+}
+
+TSoftClassPtr<UObject> UAssetDefinition_UAFBlendMask::GetAssetClass() const
+{
+	return UUAFBlendMask::StaticClass();
+}
+
+EAssetCommandResult UAssetDefinition_UAFBlendMask::OpenAssets(const FAssetOpenArgs& OpenArgs) const
+{
+	if (OpenArgs.OpenMethod == EAssetOpenMethod::Edit)
+	{
+		TArray<UObject*> Assets = OpenArgs.LoadObjects<UObject>();
+		MakeShared<UE::UAF::FBlendMaskEditorToolkit>()->InitEditor(Assets);
+	}
+
+	return EAssetCommandResult::Handled;
+}
+
+TConstArrayView<FAssetCategoryPath> UAssetDefinition_UAFBlendMask::GetAssetCategories() const
+{
+	static const auto Categories = { FAssetCategoryPath(EAssetCategoryPaths::Animation, LOCTEXT("UAFSubMenu", "Animation Framework")) };
+	return Categories;
+}
+
+#undef LOCTEXT_NAMESPACE
diff --git a/Engine/Plugins/Experimental/UAF/UAF/Source/UAFEditor/Private/UAF/BlendMask/BlendMaskAssetDefinition.h b/Engine/Plugins/Experimental/UAF/UAF/Source/UAFEditor/Private/UAF/BlendMask/BlendMaskAssetDefinition.h
new file mode 100644
--- /dev/null
+++ b/Engine/Plugins/Experimental/UAF/UAF/Source/UAFEditor/Private/UAF/BlendMask/BlendMaskAssetDefinition.h
@@ -0,0 +1,22 @@
+// Copyright Epic Games, Inc. All Rights Reserved.
+
+#pragma once
+
+#include "AssetDefinitionDefault.h"
+
+#include "BlendMaskAssetDefinition.generated.h"
+
+UCLASS()
+class UAssetDefinition_UAFBlendMask : public UAssetDefinitionDefault
+{
+	GENERATED_BODY()
+
+public:
+	// Begin UAssetDefinition
+	virtual FText GetAssetDisplayName() const override;
+	virtual FLinearColor GetAssetColor() const override;
+	virtual TSoftClassPtr<UObject> GetAssetClass() const override;
+	virtual EAssetCommandResult OpenAssets(const FAssetOpenArgs& OpenArgs) const override;
+	virtual TConstArrayView<FAssetCategoryPath> GetAssetCategories() const override;
+	// End UAssetDefinition
+};
\ No newline at end of file
diff --git a/Engine/Plugins/Experimental/UAF/UAF/Source/UAFEditor/Private/UAF/BlendMask/BlendMaskEditorToolkit.cpp b/Engine/Plugins/Experimental/UAF/UAF/Source/UAFEditor/Private/UAF/BlendMask/BlendMaskEditorToolkit.cpp
new file mode 100644
--- /dev/null
+++ b/Engine/Plugins/Experimental/UAF/UAF/Source/UAFEditor/Private/UAF/BlendMask/BlendMaskEditorToolkit.cpp
@@ -0,0 +1,144 @@
+// Copyright Epic Games, Inc. All Rights Reserved.
+
+#include "BlendMaskEditorToolkit.h"
+
+#include "HierarchyTableEditorModule.h"
+#include "Modules/ModuleManager.h"
+#include "PropertyEditorModule.h"
+#include "ToolMenus.h"
+#include "Widgets/Docking/SDockTab.h"
+
+#define LOCTEXT_NAMESPACE "BlendMaskEditorToolkit"
+
+namespace UE::UAF
+{
+
+FText FBlendMaskEditorToolkit::GetBaseToolkitName() const
+{
+	return LOCTEXT("ToolkitName", "Blend Mask Editor");
+}
+
+void FBlendMaskEditorToolkit::InitEditor(const TArray<UObject*>& InObjects)
+{
+	BlendMask = CastChecked<UUAFBlendMask>(InObjects[0]);
+	
+	BlendMask->Table->SetFlags(RF_Transactional);
+	BlendMask->UpdateHierarchy();
+
+	const TSharedRef<FTabManager::FLayout> Layout = FTabManager::NewLayout("BlendMaskEditorToolkit")
+		->AddArea
+		(
+			FTabManager::NewPrimaryArea()->SetOrientation(Orient_Horizontal)
+			->Split
+			(
+				FTabManager::NewStack()
+				->SetSizeCoefficient(0.7f)
+				->AddTab("BlendMaskEditorTableTab", ETabState::OpenedTab)
+			)
+			->Split
+			(
+				FTabManager::NewStack()
+				->SetSizeCoefficient(0.3f)
+				->AddTab("BlendMaskEditorDetailsTab", ETabState::OpenedTab)
+			)
+		);
+
+	FAssetEditorToolkit::InitAssetEditor(EToolkitMode::Standalone, {}, "BlendMaskEditor", Layout, true, true, InObjects);
+
+	ExtendToolbar();
+
+	TObjectPtr<USkeleton> Skeleton = BlendMask->GetSkeleton();
+	if (Skeleton)
+	{
+		Skeleton->RegisterOnSkeletonHierarchyChanged(USkeleton::FOnSkeletonHierarchyChanged::CreateRaw(this, &FBlendMaskEditorToolkit::OnSkeletonHierarchyChanged));
+	}
+}
+
+void FBlendMaskEditorToolkit::OnClose()
+{
+	TObjectPtr<USkeleton> Skeleton = BlendMask->GetSkeleton();
+	if (Skeleton)
+	{
+		Skeleton->UnregisterOnSkeletonHierarchyChanged(this);
+	}
+}
+
+void FBlendMaskEditorToolkit::RegisterTabSpawners(const TSharedRef<class FTabManager>& InTabManager)
+{
+	FAssetEditorToolkit::RegisterTabSpawners(InTabManager);
+
+	WorkspaceMenuCategory = InTabManager->AddLocalWorkspaceMenuCategory(LOCTEXT("BlendMaskStandaloneEditor", "Blend Mask Editor"));
+
+	// Table view
+	{
+		FHierarchyTableEditorModule& HierarchyTableModule = FModuleManager::GetModuleChecked<FHierarchyTableEditorModule>("HierarchyTableEditor");
+
+		TSharedRef<IHierarchyTable> TableWidget = HierarchyTableModule.CreateHierarchyTableWidget(BlendMask->Table);
+		HierarchyTableWidgetInterface = TableWidget;
+
+		InTabManager->RegisterTabSpawner("BlendMaskEditorTableTab", FOnSpawnTab::CreateLambda([this, TableWidget](const FSpawnTabArgs&)
+			{
+				return SNew(SDockTab)
+					[
+						TableWidget
+					];
+			}))
+			.SetDisplayName(LOCTEXT("BlendMask", "Blend Mask"))
+			.SetGroup(WorkspaceMenuCategory.ToSharedRef());
+	}
+
+	// Details panel
+	{
+		FPropertyEditorModule& PropertyEditorModule = FModuleManager::GetModuleChecked<FPropertyEditorModule>("PropertyEditor");
+		
+		FDetailsViewArgs DetailsViewArgs;
+		DetailsViewArgs.NameAreaSettings = FDetailsViewArgs::HideNameArea;
+		TSharedRef<IDetailsView> DetailsView = PropertyEditorModule.CreateDetailView(DetailsViewArgs);
+		DetailsView->SetObjects(TArray<UObject*>{ BlendMask });
+
+		InTabManager->RegisterTabSpawner("BlendMaskEditorDetailsTab", FOnSpawnTab::CreateLambda([=](const FSpawnTabArgs&)
+			{
+				return SNew(SDockTab)
+					[
+						DetailsView
+					];
+			}))
+			.SetDisplayName(INVTEXT("Details"))
+			.SetGroup(WorkspaceMenuCategory.ToSharedRef());
+	}
+}
+
+void FBlendMaskEditorToolkit::UnregisterTabSpawners(const TSharedRef<class FTabManager>& InTabManager)
+{
+	FAssetEditorToolkit::UnregisterTabSpawners(InTabManager);
+	InTabManager->UnregisterTabSpawner("BlendMaskEditorTableTab");
+	InTabManager->UnregisterTabSpawner("BlendMaskEditorDetailsTab");
+}
+
+void FBlendMaskEditorToolkit::OnSkeletonHierarchyChanged()
+{
+	BlendMask->UpdateHierarchy();
+	BlendMask->UpdateCachedData();
+}
+
+void FBlendMaskEditorToolkit::ExtendToolbar()
+{
+	FHierarchyTableEditorModule& HierarchyTableModule = FModuleManager::GetModuleChecked<FHierarchyTableEditorModule>("HierarchyTableEditor");
+	const TObjectPtr<UHierarchyTable_TableTypeHandler> Handler = HierarchyTableModule.CreateTableHandler(BlendMask->Table);
+	if (!ensure(Handler))
+	{
+		return;
+	}
+
+	FToolMenuOwnerScoped OwnerScoped(this);
+
+	FName ParentName;
+	static const FName MenuName = GetToolMenuToolbarName(ParentName);
+
+	UToolMenu* ToolMenu = UToolMenus::Get()->ExtendMenu(MenuName);
+	Handler->ExtendToolbar(ToolMenu, *HierarchyTableWidgetInterface);
+}
+
+}
+
+#undef LOCTEXT_NAMESPACE
\ No newline at end of file
diff --git a/Engine/Plugins/Experimental/UAF/UAF/Source/UAFEditor/Private/UAF/BlendMask/BlendMaskEditorToolkit.h b/Engine/Plugins/Experimental/UAF/UAF/Source/UAFEditor/Private/UAF/BlendMask/BlendMaskEditorToolkit.h
new file mode 100644
--- /dev/null
+++ b/Engine/Plugins/Experimental/UAF/UAF/Source/UAFEditor/Private/UAF/BlendMask/BlendMaskEditorToolkit.h
@@ -0,0 +1,41 @@
+// Copyright Epic Games, Inc. All Rights Reserved.
+
+#pragma once
+
+#include "CoreMinimal.h"
+#include "UAF/BlendMask/UAFBlendMask.h"
+#include "Toolkits/AssetEditorToolkit.h"
+
+class IHierarchyTable;
+
+namespace UE::UAF
+{
+
+class FBlendMaskEditorToolkit : public FAssetEditorToolkit
+{
+public:
+	void InitEditor(const TArray<UObject*>& InObjects);
+
+	// Begin FAssetEditorToolkit
+	void OnClose() override;
+
+	void RegisterTabSpawners(const TSharedRef<class FTabManager>& TabManager) override;
+	void UnregisterTabSpawners(const TSharedRef<class FTabManager>& TabManager) override;
+
+	FName GetToolkitFName() const override { return "BlendMaskEditor"; }
+	FText GetBaseToolkitName() const override;
+	FString GetWorldCentricTabPrefix() const override { return "Blend Mask "; }
+	FLinearColor GetWorldCentricTabColorScale() const override { return {}; }
+	// End FAssetEditorToolkit
+
+private:
+	void OnSkeletonHierarchyChanged();
+
+	void ExtendToolbar();
+
+	TObjectPtr<UUAFBlendMask> BlendMask;
+
+	TSharedPtr<IHierarchyTable> HierarchyTableWidgetInterface;
+};
+
+}
\ No newline at end of file
diff --git a/Engine/Plugins/Experimental/UAF/UAF/Source/UAFEditor/Private/UAF/BlendMask/BlendMaskFactory.cpp b/Engine/Plugins/Experimental/UAF/UAF/Source/UAFEditor/Private/UAF/BlendMask/BlendMaskFactory.cpp
new file mode 100644
--- /dev/null
+++ b/Engine/Plugins/Experimental/UAF/UAF/Source/UAFEditor/Private/UAF/BlendMask/BlendMaskFactory.cpp
@@ -0,0 +1,53 @@
+// Copyright Epic Games, Inc. All Rights Reserved.
+
+#include "BlendMaskFactory.h"
+#include "MaskProfile/HierarchyTableTypeMask.h"
+#include "HierarchyTableEditorModule.h"
+#include "SkeletonHierarchyTableType.h"
+#include "HierarchyTableEditorModule.h"
+#include "Modules/ModuleManager.h"
+
+#include "UAF/BlendMask/UAFBlendMask.h"
+
+#include UE_INLINE_GENERATED_CPP_BY_NAME(BlendMaskFactory)
+
+#define LOCTEXT_NAMESPACE "UAFBlendMaskFactory"
+
+UUAFBlendMaskFactory::UUAFBlendMaskFactory()
+{
+	SupportedClass = UUAFBlendMask::StaticClass();
+	bCreateNew = true;
+}
+
+UObject* UUAFBlendMaskFactory::FactoryCreateNew(UClass* Class, UObject* InParent, FName Name, EObjectFlags Flags, UObject* Context, FFeedbackContext* Warn)
+{
+	TObjectPtr<UUAFBlendMask> NewBlendMask = NewObject<UUAFBlendMask>(InParent, Class, Name, Flags, Context);
+	NewBlendMask->Table = NewObject<UHierarchyTable>(NewBlendMask);
+
+	// TODO: Streamline hierarchy table creation API
+
+	NewBlendMask->Table->Initialize(TableMetadata, FHierarchyTable_ElementType_Mask::StaticStruct());
+
+	check(TableHandler);
+	TableHandler->SetHierarchyTable(NewBlendMask->Table);
+	TableHandler->ConstructHierarchy();
+
+	return NewBlendMask;
+}
+
+bool UUAFBlendMaskFactory::ConfigureProperties()
+{
+	FHierarchyTableEditorModule& HierarchyTableModule = FModuleManager::GetModuleChecked<FHierarchyTableEditorModule>("HierarchyTableEditor");
+
+	TableHandler = HierarchyTableModule.CreateTableHandler(FHierarchyTable_TableType_Skeleton::StaticStruct());
+	check(TableHandler);
+
+	TableMetadata = FInstancedStruct(FHierarchyTable_TableType_Skeleton::StaticStruct());
+
+	// Displays window for setting skeleton
+	const bool bSuccess = TableHandler->FactoryConfigureProperties(TableMetadata);
+
+	return bSuccess;
+}
+
+#undef LOCTEXT_NAMESPACE
diff --git a/Engine/Plugins/Experimental/UAF/UAF/Source/UAFEditor/Private/UAF/BlendMask/BlendMaskFactory.h b/Engine/Plugins/Experimental/UAF/UAF/Source/UAFEditor/Private/UAF/BlendMask/BlendMaskFactory.h
new file mode 100644
--- /dev/null
+++ b/Engine/Plugins/Experimental/UAF/UAF/Source/UAFEditor/Private/UAF/BlendMask/BlendMaskFactory.h
@@ -0,0 +1,30 @@
+// Copyright Epic Games, Inc. All Rights Reserved.
+
+#pragma once
+
+#include "CoreMinimal.h"
+#include "Factories/Factory.h"
+#include "StructUtils/InstancedStruct.h"
+
+#include "BlendMaskFactory.generated.h"
+
+class UHierarchyTable_TableTypeHandler;
+
+UCLASS()
+class UUAFBlendMaskFactory : public UFactory
+{
+	GENERATED_BODY()
+
+public:
+	UUAFBlendMaskFactory();
+
+	// UFactory interface
+	UObject* FactoryCreateNew(UClass* Class, UObject* InParent, FName Name, EObjectFlags Flags, UObject* Context, FFeedbackContext* Warn) override;
+	virtual bool ConfigureProperties() override;
+
+private:
+	UPROPERTY()
+	TObjectPtr<UHierarchyTable_TableTypeHandler> TableHandler;
+
+	FInstancedStruct TableMetadata;
+};
diff --git a/Engine/Plugins/Experimental/UAF/UAF/Source/UAFEditor/Private/UAF/BlendProfile/BlendProfileAssetDefinition.cpp b/Engine/Plugins/Experimental/UAF/UAF/Source/UAFEditor/Private/UAF/BlendProfile/BlendProfileAssetDefinition.cpp
new file mode 100644
--- /dev/null
+++ b/Engine/Plugins/Experimental/UAF/UAF/Source/UAFEditor/Private/UAF/BlendProfile/BlendProfileAssetDefinition.cpp
@@ -0,0 +1,75 @@
+// Copyright Epic Games, Inc. All Rights Reserved.
+
+#include "BlendProfileAssetDefinition.h"
+
+#include "AssetDefinitionAssetInfo.h"
+#include "IAssetStatusInfoProvider.h"
+#include "UAFStyle.h"
+#include "UAF/BlendProfile/UAFBlendProfile.h"
+#include "UAF/BlendProfile/BlendProfileEditorToolkit.h"
+
+#include UE_INLINE_GENERATED_CPP_BY_NAME(BlendProfileAssetDefinition)
+
+#define LOCTEXT_NAMESPACE "UAFBlendProfile"
+
+FText UAssetDefinition_UAFBlendProfile::GetAssetDisplayName() const
+{
+	return LOCTEXT("AssetDisplayName", "UAF Blend Profile");
+}
+
+FLinearColor UAssetDefinition_UAFBlendProfile::GetAssetColor() const
+{
+	return FLinearColor(FColor::Yellow);
+}
+
+TSoftClassPtr<UObject> UAssetDefinition_UAFBlendProfile::GetAssetClass() const
+{
+	return UUAFBlendProfile::StaticClass();
+}
+
+EAssetCommandResult UAssetDefinition_UAFBlendProfile::OpenAssets(const FAssetOpenArgs& OpenArgs) const
+{
+	if (OpenArgs.OpenMethod == EAssetOpenMethod::Edit)
+	{
+		TArray<UObject*> Assets = OpenArgs.LoadObjects<UObject>();
+		MakeShared<UE::UAF::FBlendProfileEditorToolkit>()->InitEditor(Assets);
+	}
+
+	return EAssetCommandResult::Handled;
+}
+
+TConstArrayView<FAssetCategoryPath> UAssetDefinition_UAFBlendProfile::GetAssetCategories() const
+{
+	static const auto Categories = { FAssetCategoryPath(EAssetCategoryPaths::Animation, LOCTEXT("UAFSubMenu", "Animation Framework")) };
+	return Categories;
+}
+
+void UAssetDefinition_UAFBlendProfile::GetAssetStatusInfo(const TSharedPtr<IAssetStatusInfoProvider>& InAssetStatusInfoProvider, TArray<FAssetDisplayInfo>& OutStatusInfo) const
+{
+	using namespace UE::UAF;
+	
+	Super::GetAssetStatusInfo(InAssetStatusInfoProvider, OutStatusInfo);
+
+	const FAssetData AssetData = InAssetStatusInfoProvider->TryGetAssetData();
+	if (AssetData.IsValid())
+	{
+		FString BlendProfileType;
+		if (AssetData.GetTagValue(AssetRegistryTag_BlendProfileType, BlendProfileType))
+		{
+			FAssetDisplayInfo TypeStatus;
+			TypeStatus.Priority = FAssetStatusPriority(EStatusSeverity::Info);
+			TypeStatus.IsVisible = EVisibility::Visible;
+			TypeStatus.StatusIcon = (BlendProfileType == AssetRegistryTag_BlendProfileType_TimeProfile)
+				? FUAFStyle::Get().GetBrush("BlendProfiles.TimeProfile")
+				: FUAFStyle::Get().GetBrush("BlendProfiles.WeightProfile");
+			TypeStatus.StatusDescription = (BlendProfileType == AssetRegistryTag_BlendProfileType_TimeProfile)
+				? LOCTEXT("AssetDefinition_UAFBlendProfile_TimeProfileTypeTooltip", "Time Blend Profile")
+				: LOCTEXT("AssetDefinition_UAFBlendProfile_WeightProfileTypeTooltip", "Weight Blend Profile");
+	
+			OutStatusInfo.Add(MoveTemp(TypeStatus));
+		}
+	}
+
+}
+
+#undef LOCTEXT_NAMESPACE
diff --git a/Engine/Plugins/Experimental/UAF/UAF/Source/UAFEditor/Private/UAF/BlendProfile/BlendProfileAssetDefinition.h b/Engine/Plugins/Experimental/UAF/UAF/Source/UAFEditor/Private/UAF/BlendProfile/BlendProfileAssetDefinition.h
new file mode 100644
--- /dev/null
+++ b/Engine/Plugins/Experimental/UAF/UAF/Source/UAFEditor/Private/UAF/BlendProfile/BlendProfileAssetDefinition.h
@@ -0,0 +1,23 @@
+// Copyright Epic Games, Inc. All Rights Reserved.
+
+#pragma once
+
+#include "AssetDefinitionDefault.h"
+
+#include "BlendProfileAssetDefinition.generated.h"
+
+UCLASS()
+class UAssetDefinition_UAFBlendProfile : public UAssetDefinitionDefault
+{
+	GENERATED_BODY()
+
+public:
+	// Begin UAssetDefinition
+	virtual FText GetAssetDisplayName() const override;
+	virtual FLinearColor GetAssetColor() const override;
+	virtual TSoftClassPtr<UObject> GetAssetClass() const override;
+	virtual EAssetCommandResult OpenAssets(const FAssetOpenArgs& OpenArgs) const override;
+	virtual TConstArrayView<FAssetCategoryPath> GetAssetCategories() const override;
+	virtual void GetAssetStatusInfo(const TSharedPtr<IAssetStatusInfoProvider>& InAssetStatusInfoProvider, TArray<FAssetDisplayInfo>& OutStatusInfo) const;
+	// End UAssetDefinition
+};
\ No newline at end of file
diff --git a/Engine/Plugins/Experimental/UAF/UAF/Source/UAFEditor/Private/UAF/BlendProfile/BlendProfileEditorToolkit.cpp b/Engine/Plugins/Experimental/UAF/UAF/Source/UAFEditor/Private/UAF/BlendProfile/BlendProfileEditorToolkit.cpp
new file mode 100644
--- /dev/null
+++ b/Engine/Plugins/Experimental/UAF/UAF/Source/UAFEditor/Private/UAF/BlendProfile/BlendProfileEditorToolkit.cpp
@@ -0,0 +1,144 @@
+// Copyright Epic Games, Inc. All Rights Reserved.
+
+#include "BlendProfileEditorToolkit.h"
+
+#include "HierarchyTableEditorModule.h"
+#include "Modules/ModuleManager.h"
+#include "PropertyEditorModule.h"
+#include "ToolMenus.h"
+#include "Widgets/Docking/SDockTab.h"
+
+#define LOCTEXT_NAMESPACE "BlendProfileEditorToolkit"
+
+namespace UE::UAF
+{
+
+FText FBlendProfileEditorToolkit::GetBaseToolkitName() const
+{
+	return LOCTEXT("ToolkitName", "Blend Profile Editor");
+}
+
+void FBlendProfileEditorToolkit::InitEditor(const TArray<UObject*>& InObjects)
+{
+	BlendProfile = CastChecked<UUAFBlendProfile>(InObjects[0]);
+	
+	BlendProfile->Table->SetFlags(RF_Transactional);
+	BlendProfile->UpdateHierarchy();
+
+	const TSharedRef<FTabManager::FLayout> Layout = FTabManager::NewLayout("BlendProfileEditorToolkit")
+		->AddArea
+		(
+			FTabManager::NewPrimaryArea()->SetOrientation(Orient_Horizontal)
+			->Split
+			(
+				FTabManager::NewStack()
+				->SetSizeCoefficient(0.7f)
+				->AddTab("BlendProfileEditorTableTab", ETabState::OpenedTab)
+			)
+			->Split
+			(
+				FTabManager::NewStack()
+				->SetSizeCoefficient(0.3f)
+				->AddTab("BlendProfileEditorDetailsTab", ETabState::OpenedTab)
+			)
+		);
+
+	FAssetEditorToolkit::InitAssetEditor(EToolkitMode::Standalone, {}, "BlendProfileEditor", Layout, true, true, InObjects);
+
+	ExtendToolbar();
+
+	TObjectPtr<USkeleton> Skeleton = BlendProfile->GetSkeleton();
+	if (Skeleton)
+	{
+		Skeleton->RegisterOnSkeletonHierarchyChanged(USkeleton::FOnSkeletonHierarchyChanged::CreateRaw(this, &FBlendProfileEditorToolkit::OnSkeletonHierarchyChanged));
+	}
+}
+
+void FBlendProfileEditorToolkit::OnClose()
+{
+	TObjectPtr<USkeleton> Skeleton = BlendProfile->GetSkeleton();
+	if (Skeleton)
+	{
+		Skeleton->UnregisterOnSkeletonHierarchyChanged(this);
+	}
+}
+
+void FBlendProfileEditorToolkit::RegisterTabSpawners(const TSharedRef<class FTabManager>& InTabManager)
+{
+	FAssetEditorToolkit::RegisterTabSpawners(InTabManager);
+
+	WorkspaceMenuCategory = InTabManager->AddLocalWorkspaceMenuCategory(LOCTEXT("BlendProfileStandaloneEditor", "Blend Profile Editor"));
+
+	// Table view
+	{
+		FHierarchyTableEditorModule& HierarchyTableModule = FModuleManager::GetModuleChecked<FHierarchyTableEditorModule>("HierarchyTableEditor");
+
+		TSharedRef<IHierarchyTable> TableWidget = HierarchyTableModule.CreateHierarchyTableWidget(BlendProfile->Table);
+		HierarchyTableWidgetInterface = TableWidget;
+
+		InTabManager->RegisterTabSpawner("BlendProfileEditorTableTab", FOnSpawnTab::CreateLambda([this, TableWidget](const FSpawnTabArgs&)
+			{
+				return SNew(SDockTab)
+					[
+						TableWidget
+					];
+			}))
+			.SetDisplayName(LOCTEXT("BlendProfile", "Blend Profile"))
+			.SetGroup(WorkspaceMenuCategory.ToSharedRef());
+	}
+
+	// Details panel
+	{
+		FPropertyEditorModule& PropertyEditorModule = FModuleManager::GetModuleChecked<FPropertyEditorModule>("PropertyEditor");
+		
+		FDetailsViewArgs DetailsViewArgs;
+		DetailsViewArgs.NameAreaSettings = FDetailsViewArgs::HideNameArea;
+		TSharedRef<IDetailsView> DetailsView = PropertyEditorModule.CreateDetailView(DetailsViewArgs);
+		DetailsView->SetObjects(TArray<UObject*>{ BlendProfile });
+
+		InTabManager->RegisterTabSpawner("BlendProfileStandaloneEditorDetailsTab", FOnSpawnTab::CreateLambda([=](const FSpawnTabArgs&)
+			{
+				return SNew(SDockTab)
+					[
+						DetailsView
+					];
+			}))
+			.SetDisplayName(INVTEXT("Details"))
+			.SetGroup(WorkspaceMenuCategory.ToSharedRef());
+	}
+}
+
+void FBlendProfileEditorToolkit::UnregisterTabSpawners(const TSharedRef<class FTabManager>& InTabManager)
+{
+	FAssetEditorToolkit::UnregisterTabSpawners(InTabManager);
+	InTabManager->UnregisterTabSpawner("BlendProfileEditorTableTab");
+	InTabManager->UnregisterTabSpawner("BlendProfileEditorDetailsTab");
+}
+
+void FBlendProfileEditorToolkit::OnSkeletonHierarchyChanged()
+{
+	BlendProfile->UpdateHierarchy();
+	BlendProfile->UpdateCachedData();
+}
+
+void FBlendProfileEditorToolkit::ExtendToolbar()
+{
+	FHierarchyTableEditorModule& HierarchyTableModule = FModuleManager::GetModuleChecked<FHierarchyTableEditorModule>("HierarchyTableEditor");
+	const TObjectPtr<UHierarchyTable_TableTypeHandler> Handler = HierarchyTableModule.CreateTableHandler(BlendProfile->Table);
+	if (!ensure(Handler))
+	{
+		return;
+	}
+
+	FToolMenuOwnerScoped OwnerScoped(this);
+
+	FName ParentName;
+	static const FName MenuName = GetToolMenuToolbarName(ParentName);
+
+	UToolMenu* ToolMenu = UToolMenus::Get()->ExtendMenu(MenuName);
+	Handler->ExtendToolbar(ToolMenu, *HierarchyTableWidgetInterface);
+}
+
+}
+
+#undef LOCTEXT_NAMESPACE
\ No newline at end of file
diff --git a/Engine/Plugins/Experimental/UAF/UAF/Source/UAFEditor/Private/UAF/BlendProfile/BlendProfileEditorToolkit.h b/Engine/Plugins/Experimental/UAF/UAF/Source/UAFEditor/Private/UAF/BlendProfile/BlendProfileEditorToolkit.h
new file mode 100644
--- /dev/null
+++ b/Engine/Plugins/Experimental/UAF/UAF/Source/UAFEditor/Private/UAF/BlendProfile/BlendProfileEditorToolkit.h
@@ -0,0 +1,41 @@
+// Copyright Epic Games, Inc. All Rights Reserved.
+
+#pragma once
+
+#include "CoreMinimal.h"
+#include "UAF/BlendProfile/UAFBlendProfile.h"
+#include "Toolkits/AssetEditorToolkit.h"
+
+class IHierarchyTable;
+
+namespace UE::UAF
+{
+
+class FBlendProfileEditorToolkit : public FAssetEditorToolkit
+{
+public:
+	void InitEditor(const TArray<UObject*>& InObjects);
+
+	// Begin FAssetEditorToolkit
+	void OnClose() override;
+
+	void RegisterTabSpawners(const TSharedRef<class FTabManager>& TabManager) override;
+	void UnregisterTabSpawners(const TSharedRef<class FTabManager>& TabManager) override;
+
+	FName GetToolkitFName() const override { return "BlendProfileEditor"; }
+	FText GetBaseToolkitName() const override;
+	FString GetWorldCentricTabPrefix() const override { return "Blend Profile "; }
+	FLinearColor GetWorldCentricTabColorScale() const override { return {}; }
+	// End FAssetEditorToolkit
+
+private:
+	void OnSkeletonHierarchyChanged();
+
+	void ExtendToolbar();
+
+	TObjectPtr<UUAFBlendProfile> BlendProfile;
+
+	TSharedPtr<IHierarchyTable> HierarchyTableWidgetInterface;
+};
+
+}
\ No newline at end of file
diff --git a/Engine/Plugins/Experimental/UAF/UAF/Source/UAFEditor/Private/UAF/BlendProfile/BlendProfileFactory.cpp b/Engine/Plugins/Experimental/UAF/UAF/Source/UAFEditor/Private/UAF/BlendProfile/BlendProfileFactory.cpp
new file mode 100644
--- /dev/null
+++ b/Engine/Plugins/Experimental/UAF/UAF/Source/UAFEditor/Private/UAF/BlendProfile/BlendProfileFactory.cpp
@@ -0,0 +1,116 @@
+// Copyright Epic Games, Inc. All Rights Reserved.
+
+#include "BlendProfileFactory.h"
+#include "HierarchyTableDefaultTypes.h"
+#include "HierarchyTableEditorModule.h"
+#include "SkeletonHierarchyTableType.h"
+#include "Modules/ModuleManager.h"
+#include "Widgets/SWindow.h"
+#include "Widgets/Input/SButton.h"
+#include "SEnumCombo.h"
+#include "Editor.h"
+
+#include "UAF/BlendProfile/UAFBlendProfile.h"
+
+#include UE_INLINE_GENERATED_CPP_BY_NAME(BlendProfileFactory)
+
+#define LOCTEXT_NAMESPACE "UAFBlendProfileFactory"
+
+UUAFBlendProfileFactory::UUAFBlendProfileFactory()
+	: BlendProfileType(EUAFBlendProfileType::WeightFactor)
+{
+	SupportedClass = UUAFBlendProfile::StaticClass();
+	bCreateNew = true;
+}
+
+UObject* UUAFBlendProfileFactory::FactoryCreateNew(UClass* Class, UObject* InParent, FName Name, EObjectFlags Flags, UObject* Context, FFeedbackContext* Warn)
+{
+	TObjectPtr<UUAFBlendProfile> NewBlendProfile = NewObject<UUAFBlendProfile>(InParent, Class, Name, Flags, Context);
+	NewBlendProfile->Table = NewObject<UHierarchyTable>(NewBlendProfile);
+	NewBlendProfile->Type = BlendProfileType;
+
+	// TODO: Streamline hierarchy table creation API
+
+	NewBlendProfile->Table->Initialize(TableMetadata, FHierarchyTable_ElementType_Float::StaticStruct());
+
+	check(TableHandler);
+	TableHandler->SetHierarchyTable(NewBlendProfile->Table);
+	TableHandler->ConstructHierarchy();
+
+	if (FHierarchyTableEntryData* RootEntry = NewBlendProfile->Table->GetMutableTableEntry(0))
+	{
+		if (ensure(RootEntry->IsOverridden()))
+		{
+			RootEntry->GetMutableValue<FHierarchyTable_ElementType_Float>()->Value = 1.0f;
+		}
+	}
+
+	return NewBlendProfile;
+}
+
+bool UUAFBlendProfileFactory::ConfigureProperties()
+{
+	return ConfigureBlendProfileType() && ConfigureBlendProfileHierarchy();
+}
+
+bool UUAFBlendProfileFactory::ConfigureBlendProfileType()
+{
+	bool bConfirmClicked = false;
+
+	TSharedRef<SWindow> Window = SNew(SWindow)
+		.Title(LOCTEXT("Title", "Choose Blend Profile Type"))
+		.ClientSize(FVector2D(400, 400))
+		.SupportsMinimize(false)
+		.SupportsMaximize(false)
+		[
+			SNew(SVerticalBox)
+				+ SVerticalBox::Slot()
+				.AutoHeight()
+				[
+					SNew(SEnumComboBox, StaticEnum<EUAFBlendProfileType>())
+						.CurrentValue_Lambda([this]()
+							{
+								return static_cast<int32>(BlendProfileType);
+							})
+						.OnEnumSelectionChanged_Lambda([this](int32 InEnumValue, ESelectInfo::Type SelectInfo)
+							{
+								BlendProfileType = static_cast<EUAFBlendProfileType>(InEnumValue);
+							})
+				]
+			+ SVerticalBox::Slot()
+				.AutoHeight()
+				[
+					SNew(SButton)
+						.OnClicked_Lambda([&]()
+							{
+								bConfirmClicked = true;
+								Window->RequestDestroyWindow();
+								return FReply::Handled();
+							})
+						[
+							SNew(STextBlock)
+								.Text(LOCTEXT("Confirm", "Confirm"))
+						]
+				]
+		];
+
+
+	GEditor->EditorAddModalWindow(Window);
+
+	return bConfirmClicked;
+}
+
+bool UUAFBlendProfileFactory::ConfigureBlendProfileHierarchy()
+{
+	FHierarchyTableEditorModule& HierarchyTableModule = FModuleManager::GetModuleChecked<FHierarchyTableEditorModule>("HierarchyTableEditor");
+
+	TableHandler = HierarchyTableModule.CreateTableHandler(FHierarchyTable_TableType_Skeleton::StaticStruct());
+	check(TableHandler);
+
+	TableMetadata = FInstancedStruct(FHierarchyTable_TableType_Skeleton::StaticStruct());
+
+	// Displays window for setting skeleton
+	return TableHandler->FactoryConfigureProperties(TableMetadata);
+}
+
+#undef LOCTEXT_NAMESPACE
diff --git a/Engine/Plugins/Experimental/UAF/UAF/Source/UAFEditor/Private/UAF/BlendProfile/BlendProfileFactory.h b/Engine/Plugins/Experimental/UAF/UAF/Source/UAFEditor/Private/UAF/BlendProfile/BlendProfileFactory.h
new file mode 100644
--- /dev/null
+++ b/Engine/Plugins/Experimental/UAF/UAF/Source/UAFEditor/Private/UAF/BlendProfile/BlendProfileFactory.h
@@ -0,0 +1,37 @@
+// Copyright Epic Games, Inc. All Rights Reserved.
+
+#pragma once
+
+#include "CoreMinimal.h"
+#include "Factories/Factory.h"
+#include "StructUtils/InstancedStruct.h"
+#include "UAF/BlendProfile/UAFBlendProfile.h"
+
+#include "BlendProfileFactory.generated.h"
+
+class UHierarchyTable_TableTypeHandler;
+
+UCLASS()
+class UUAFBlendProfileFactory : public UFactory
+{
+	GENERATED_BODY()
+
+public:
+	UUAFBlendProfileFactory();
+
+	// UFactory interface
+	UObject* FactoryCreateNew(UClass* Class, UObject* InParent, FName Name, EObjectFlags Flags, UObject* Context, FFeedbackContext* Warn) override;
+	virtual bool ConfigureProperties() override;
+
+private:
+	bool ConfigureBlendProfileType();
+
+	bool ConfigureBlendProfileHierarchy();
+
+	UPROPERTY()
+	TObjectPtr<UHierarchyTable_TableTypeHandler> TableHandler;
+
+	EUAFBlendProfileType BlendProfileType;
+
+	FInstancedStruct TableMetadata;
+};
diff --git a/Engine/Plugins/Experimental/UAF/UAF/Source/UAFEditor/Private/UAFStyle.cpp b/Engine/Plugins/Experimental/UAF/UAF/Source/UAFEditor/Private/UAFStyle.cpp
--- a/Engine/Plugins/Experimental/UAF/UAF/Source/UAFEditor/Private/UAFStyle.cpp
+++ b/Engine/Plugins/Experimental/UAF/UAF/Source/UAFEditor/Private/UAFStyle.cpp
@@ -26,6 +26,7 @@ FUAFStyle::FUAFStyle()
 	SetTraitStyles();
 	SetAssetWizardStyles();
 	SetSetBindingStyles();
+	SetBlendProfileStyles();
 
 	FSlateStyleRegistry::RegisterSlateStyle(*this);
 }
@@ -116,6 +117,14 @@ void FUAFStyle::SetSetBindingStyles()
 	Set("SetBindingEditorCommands.UnbindSelection", new CORE_IMAGE_BRUSH_SVG("Starship/Common/Unlinked", Icon16x16));
 }
 
+void FUAFStyle::SetBlendProfileStyles()
+{
+	const FVector2D Icon16x16(16.0f, 16.0f);
+
+	Set("BlendProfiles.TimeProfile", new IMAGE_BRUSH_SVG("Starship/Common/Realtime", Icon16x16));
+	Set("BlendProfiles.WeightProfile", new FSlateVectorImageBrush(InResources("Icons/Weight_16.svg"),Icon16x16));
+}
+
 FUAFStyle& FUAFStyle::Get()
 {
 	static FUAFStyle Instance;
diff --git a/Engine/Plugins/Experimental/UAF/UAF/Source/UAFEditor/UAFEditor.Build.cs b/Engine/Plugins/Experimental/UAF/UAF/Source/UAFEditor/UAFEditor.Build.cs
--- a/Engine/Plugins/Experimental/UAF/UAF/Source/UAFEditor/UAFEditor.Build.cs
+++ b/Engine/Plugins/Experimental/UAF/UAF/Source/UAFEditor/UAFEditor.Build.cs
@@ -64,7 +64,10 @@ namespace UnrealBuildTool.Rules
 					"Projects",
 					"Persona",
 					"DeveloperSettings",
-					"AdvancedPreviewScene"
+					"AdvancedPreviewScene",
+					"HierarchyTableRuntime",
+					"HierarchyTableEditor",
+					"HierarchyTableAnimationRuntime",
 				}
 			);
 			
diff --git a/Engine/Plugins/Experimental/UAF/UAFAnimGraph/Source/UAFAnimGraph/Internal/Injection/MontageTrait.h b/Engine/Plugins/Experimental/UAF/UAFAnimGraph/Source/UAFAnimGraph/Internal/Injection/MontageTrait.h
--- a/Engine/Plugins/Experimental/UAF/UAFAnimGraph/Source/UAFAnimGraph/Internal/Injection/MontageTrait.h
+++ b/Engine/Plugins/Experimental/UAF/UAFAnimGraph/Source/UAFAnimGraph/Internal/Injection/MontageTrait.h
@@ -56,7 +56,7 @@ GENERATE_TRAIT_LATENT_PROPERTIES(FUAFMontageTraitSharedData, TRAIT_LATENT_PROPER
 
 namespace UE::UAF
 {
-	struct FMontageTrait : FBaseTrait, IUpdate, IEvaluate, IHierarchy, IUpdateTraversal, ITimeline, IGroupSynchronization
+	struct FMontageTrait : FBaseTrait, IUpdate, IEvaluate, IHierarchy, IUpdateTraversal, ITimeline, IGroupSynchronization, IGarbageCollection
 	{
 		DECLARE_ANIM_TRAIT(FMontageTrait, FBaseTrait)
 
@@ -128,6 +128,9 @@ namespace UE::UAF
 			/** Blend data for additive poses if blend profiles are used */
 			TArray<FBlendSampleData> AdditiveBlendDataForBlendProfiles;
 
+			/** The skeleton that owns each blend profile used, or nullptr if montage doesn't use a blend profile */
+			TArray<TObjectPtr<const USkeleton>> SkeletonDataForBlendProfiles;
+
 			void Construct(const FExecutionContext& Context, const FTraitBinding& Binding);
 			void ResetData();
 		};
@@ -154,5 +157,8 @@ namespace UE::UAF
 		// IGroupSynchronization impl
 		virtual FSyncGroupParameters GetGroupParameters(FExecutionContext& Context, const TTraitBinding<IGroupSynchronization>& Binding) const override;
 		virtual void AdvanceBy(FExecutionContext& Context, const TTraitBinding<IGroupSynchronization>& Binding, float DeltaTime, bool bDispatchEvents) const override;
+
+		// IGarbageCollection impl
+		virtual void AddReferencedObjects(const FExecutionContext& Context, const TTraitBinding<IGarbageCollection>& Binding, FReferenceCollector& Collector) const override;
 	};
 }
diff --git a/Engine/Plugins/Experimental/UAF/UAFAnimGraph/Source/UAFAnimGraph/Internal/Traits/BlendStackTrait.h b/Engine/Plugins/Experimental/UAF/UAFAnimGraph/Source/UAFAnimGraph/Internal/Traits/BlendStackTrait.h
--- a/Engine/Plugins/Experimental/UAF/UAFAnimGraph/Source/UAFAnimGraph/Internal/Traits/BlendStackTrait.h
+++ b/Engine/Plugins/Experimental/UAF/UAFAnimGraph/Source/UAFAnimGraph/Internal/Traits/BlendStackTrait.h
@@ -209,7 +209,7 @@ namespace UE::UAF
 		virtual FAnimNextGraphInstance* GetGraphInstance(FExecutionContext& Context, const TTraitBinding<IBlendStack>& Binding, int32 InChildIndex) const override;
 
 		// ISmoothBlendPerBone impl
-		virtual TSharedPtr<const IBlendProfileInterface> GetBlendProfile(FExecutionContext& Context, const TTraitBinding<ISmoothBlendPerBone>& Binding, int32 ChildIndex) const;
+		virtual const UUAFBlendProfile* GetBlendProfile(FExecutionContext& Context, const TTraitBinding<ISmoothBlendPerBone>& Binding, int32 ChildIndex) const;
 
 #if WITH_EDITOR
 		virtual bool IsHidden() const override { return true; }
diff --git a/Engine/Plugins/Experimental/UAF/UAFAnimGraph/Source/UAFAnimGraph/Private/EvaluationVM/Tasks/ApplyAdditiveKeyframePerBone.cpp b/Engine/Plugins/Experimental/UAF/UAFAnimGraph/Source/UAFAnimGraph/Private/EvaluationVM/Tasks/ApplyAdditiveKeyframePerBone.cpp
--- a/Engine/Plugins/Experimental/UAF/UAFAnimGraph/Source/UAFAnimGraph/Private/EvaluationVM/Tasks/ApplyAdditiveKeyframePerBone.cpp
+++ b/Engine/Plugins/Experimental/UAF/UAFAnimGraph/Source/UAFAnimGraph/Private/EvaluationVM/Tasks/ApplyAdditiveKeyframePerBone.cpp
@@ -2,19 +2,17 @@
 
 #include "EvaluationVM/Tasks/ApplyAdditiveKeyframePerBone.h"
 
-#include "HierarchyTable.h"
-#include "HierarchyTableBlendProfile.h"
 #include "EvaluationVM/EvaluationVM.h"
 #include "EvaluationVM/KeyframeState.h"
 #include "TransformArrayOperations.h"
 #include "Animation/AnimationAsset.h"
 #include "UAF/ValueRuntime/ValueBundle.h"
 #include "Animation/AttributeTypes.h"
-#include "Animation/BlendProfile.h"
 #include "Animation/SkeletonRemappingRegistry.h"
 #include "MaskProfile/HierarchyTableTypeMask.h"
 #include "UAF/ValueRuntime/Transformers/AdditiveSpace.h"
 #include "UAF/ValueRuntime/Transformers/BoneSpace.h"
+#include "UAF/BlendMask/UAFBlendMask.h"
 
 #include UE_INLINE_GENERATED_CPP_BY_NAME(ApplyAdditiveKeyframePerBone)
 
@@ -235,39 +233,30 @@ void FUAFApplyAdditiveKeyframePerBoneTask::Execute(UE::UAF::FEvaluationVM& VM) c
 	}
 }
 
-FUAFApplyAdditiveKeyframeWithBlendProfileTask FUAFApplyAdditiveKeyframeWithBlendProfileTask::Make(const float InBlendWeight, const TSharedPtr<const IBlendProfileInterface>& InBlendProfileInterface)
+FUAFApplyAdditiveKeyframeWithBlendMaskTask FUAFApplyAdditiveKeyframeWithBlendMaskTask::Make(const float InBlendWeight, const UUAFBlendMask* InBlendMask)
 {
-	FUAFApplyAdditiveKeyframeWithBlendProfileTask Task;
+	FUAFApplyAdditiveKeyframeWithBlendMaskTask Task;
 	Task.BlendWeight = InBlendWeight;
-	Task.BlendProfileInterface = InBlendProfileInterface;
+	Task.BlendMask = InBlendMask;
 	return Task;
 }
 
-void FUAFApplyAdditiveKeyframeWithBlendProfileTask::Execute(UE::UAF::FEvaluationVM& VM) const
+void FUAFApplyAdditiveKeyframeWithBlendMaskTask::Execute(UE::UAF::FEvaluationVM& VM) const
 {
 	using namespace UE::UAF;
 
-	if (!BlendProfileInterface.IsValid())
+	if (!BlendMask)
 	{
 		// If we do not have a valid blend profile - fall back to a full body additive
-		UE_LOG(LogAnimation, Warning, TEXT("UAFApplyAdditiveKeyframeWithBlendProfileTask::Execute: No valid blend profile provided. Falling back to full body additive."));
+		UE_LOG(LogAnimation, Warning, TEXT("UAFApplyAdditiveKeyframeWithBlendMaskTask::Execute: No valid blend profile provided. Falling back to full body additive."));
 		Super::Execute(VM);
 		return;
 	}
 	
-	// Currently this task only supports HierarchyTable based blend profiles
-	const TSharedPtr<const FHierarchyTableBlendProfile> HTBlendProfile = StaticCastSharedPtr<const FHierarchyTableBlendProfile>(BlendProfileInterface);
-	if (!HTBlendProfile.IsValid())
-	{
-		// If we do not have a valid blend profile - fall back to a full body additive
-		UE_LOG(LogAnimation, Warning, TEXT("UAFApplyAdditiveKeyframeWithBlendProfileTask::Execute: This task currently only supports HierarchyTable based blend profiles. Falling back to full body additive."));
-		Super::Execute(VM);
-		return;
-	}
-	
-	const USkeleton* SourceSkeleton = BlendProfileInterface->GetSkeleton();
-	const UE::Anim::TNamedValueArray<FDefaultAllocator, UE::Anim::FCurveElement>& CurveMaskWeights = HTBlendProfile->GetCurveBlendWeights(); 
-	const TArray<FHierarchyTableBlendProfile::FMaskedAttributeWeight>& AttributeMaskWeights = HTBlendProfile->GetAttributeBlendWeights();
+	const USkeleton* SourceSkeleton = BlendMask->GetSkeleton();
+	const UUAFBlendMask::FSkeletonBoneWeightArray& BoneMaskWeights = BlendMask->GetSkeletonBoneWeights(); 
+	const UUAFBlendMask::FSkeletonCurveWeightArray& CurveMaskWeights = BlendMask->GetSkeletonCurveWeights(); 
+	const UUAFBlendMask::FSkeletonAttributeWeightArray& AttributeMaskWeights = BlendMask->GetSkeletonAttributeWeights();
 	
 	if (VM.GetActiveNamedSet())
 	{
@@ -288,7 +277,7 @@ void FUAFApplyAdditiveKeyframeWithBlendProfileTask::Execute(UE::UAF::FEvaluation
 			Collection.InitWithValueSpace(FValueSpace(EValueSpaceType::Local));
 
 			VM.PushValue(ATTRIBUTE_STACK_NAME, MakeUnique<FValueBundle>(MoveTemp(Collection)));
-			UE_LOG(LogAnimation, Warning, TEXT("UAFApplyAdditiveKeyframeWithBlendProfileTask::Execute: Could not apply additive keyframe - only a single input was provided. Pushing a ref pose instead."));
+			UE_LOG(LogAnimation, Warning, TEXT("UAFApplyAdditiveKeyframeWithBlendMaskTask::Execute: Could not apply additive keyframe - only a single input was provided. Pushing a ref pose instead."));
 			return;
 		}
 
@@ -300,7 +289,7 @@ void FUAFApplyAdditiveKeyframeWithBlendProfileTask::Execute(UE::UAF::FEvaluation
 			Collection.InitWithValueSpace(FValueSpace(EValueSpaceType::Local));
 
 			VM.PushValue(ATTRIBUTE_STACK_NAME, MakeUnique<FValueBundle>(MoveTemp(Collection)));
-			UE_LOG(LogAnimation, Warning, TEXT("UAFApplyAdditiveKeyframeWithBlendProfileTask::Execute: Could not apply additive keyframe - The expected additive pose was not of an additive type. Pushing a ref pose instead."));
+			UE_LOG(LogAnimation, Warning, TEXT("UAFApplyAdditiveKeyframeWithBlendMaskTask::Execute: Could not apply additive keyframe - The expected additive pose was not of an additive type. Pushing a ref pose instead."));
 			return;
 		}
 		
@@ -327,7 +316,7 @@ void FUAFApplyAdditiveKeyframeWithBlendProfileTask::Execute(UE::UAF::FEvaluation
 				const FAttributeTypedSetPtr& BoneSet = BoneMap->GetTypedSet();
 				const int32 NumBonesInSet = BoneMap->Num();
 
-				check(BlendProfileInterface->GetNumBlendEntries() >= NumBonesInSet);
+				check(BlendMask->GetSkeletonBoneWeights().Num() >= NumBonesInSet);
 				
 				const USkeleton* TargetSkeleton = VM.GetActiveEvaluationContext().GetSkeleton();
 				const FSkeletonRemapping& SkeletonRemapping = UE::Anim::FSkeletonRemappingRegistry::Get().GetRemapping(SourceSkeleton, TargetSkeleton);
@@ -345,7 +334,7 @@ void FUAFApplyAdditiveKeyframeWithBlendProfileTask::Execute(UE::UAF::FEvaluation
 					}
 					;
 					
-					(*SetWeights)[SetIndex].Value = WeightIndex != INDEX_NONE ? BlendProfileInterface->GetBoneBlendScale(WeightIndex) : 0.0f;
+					(*SetWeights)[SetIndex].Value = WeightIndex != INDEX_NONE ? BoneMaskWeights[WeightIndex] : 0.0f;
 					(*SetWeights)[SetIndex].Value *= AdditiveWeight;
 				}
 			}
@@ -367,7 +356,7 @@ void FUAFApplyAdditiveKeyframeWithBlendProfileTask::Execute(UE::UAF::FEvaluation
 					});
 				
 				// Attributes
-				for (const FHierarchyTableBlendProfile::FMaskedAttributeWeight& AttributeWeight : AttributeMaskWeights)
+				for (const FBlendProfileStandaloneCachedData::FMaskedAttributeWeight& AttributeWeight : AttributeMaskWeights)
 				{
 					const FAttributeSetIndex SetIndex = FloatSet->FindIndex(AttributeWeight.Attribute.GetName());
 					if (SetIndex.IsValid())
@@ -382,7 +371,7 @@ void FUAFApplyAdditiveKeyframeWithBlendProfileTask::Execute(UE::UAF::FEvaluation
 				SetWeights->FillWithValue(FFloatAnimationAttribute{ 0.0f });
 
 				const FAttributeTypedSetPtr& TypedSet = It.GetMap()->GetTypedSet();
-				for (const FHierarchyTableBlendProfile::FMaskedAttributeWeight& AttributeWeight : AttributeMaskWeights)
+				for (const FBlendProfileStandaloneCachedData::FMaskedAttributeWeight& AttributeWeight : AttributeMaskWeights)
 				{
 					const FAttributeSetIndex SetIndex = TypedSet->FindIndex(AttributeWeight.Attribute.GetName());
 					if (SetIndex.IsValid())
@@ -431,7 +420,7 @@ void FUAFApplyAdditiveKeyframeWithBlendProfileTask::Execute(UE::UAF::FEvaluation
 			// We have a single input, discard it since it must be the additive pose, either way something went wrong
 			// Push the reference pose since we'll expect a non-additive pose
 			VM.PushValue(KEYFRAME_STACK_NAME, MakeUnique<FKeyframeState>(VM.MakeReferenceKeyframe(false)));
-			UE_LOG(LogAnimation, Warning, TEXT("UAFApplyAdditiveKeyframeWithBlendProfileTask::Execute: Could not apply additive keyframe - only a single input was provided. Pushing a ref pose instead."));
+			UE_LOG(LogAnimation, Warning, TEXT("UAFApplyAdditiveKeyframeWithBlendMaskTask::Execute: Could not apply additive keyframe - only a single input was provided. Pushing a ref pose instead."));
 			return;
 		}
 
@@ -439,7 +428,7 @@ void FUAFApplyAdditiveKeyframeWithBlendProfileTask::Execute(UE::UAF::FEvaluation
 		{
 			// Additive must be additive type, push reference pose if not the case
 			VM.PushValue(KEYFRAME_STACK_NAME, MakeUnique<FKeyframeState>(VM.MakeReferenceKeyframe(false)));
-			UE_LOG(LogAnimation, Warning, TEXT("UAFApplyAdditiveKeyframeWithBlendProfileTask::Execute: Could not apply additive keyframe - The expected additive pose was not of an additive type. Pushing a ref pose instead."));
+			UE_LOG(LogAnimation, Warning, TEXT("UAFApplyAdditiveKeyframeWithBlendMaskTask::Execute: Could not apply additive keyframe - The expected additive pose was not of an additive type. Pushing a ref pose instead."));
 			return;
 		}
 
@@ -452,7 +441,6 @@ void FUAFApplyAdditiveKeyframeWithBlendProfileTask::Execute(UE::UAF::FEvaluation
 		LODBoneIndexToWeightIndexMap.AddUninitialized(NumLODBones);
 		
 		const USkeleton* TargetSkeleton = BaseKeyframe->Pose.GetSkeletonAsset();
-		TSharedPtr<IInterpolationIndexProvider::FPerBoneInterpolationData> Data = HTBlendProfile->GetPerBoneInterpolationData(TargetSkeleton);
 		const FSkeletonRemapping& SkeletonRemapping = UE::Anim::FSkeletonRemappingRegistry::Get().GetRemapping(SourceSkeleton, TargetSkeleton);
 		if (SkeletonRemapping.IsValid())
 		{
@@ -461,7 +449,7 @@ void FUAFApplyAdditiveKeyframeWithBlendProfileTask::Execute(UE::UAF::FEvaluation
 				const FSkeletonPoseBoneIndex TargetSkeletonBoneIndex(LODBoneIndexToSkeletonBoneIndexMap[LODBoneIndex]);
 				const FSkeletonPoseBoneIndex SourceSkeletonBoneIndex(SkeletonRemapping.GetSourceSkeletonBoneIndex(TargetSkeletonBoneIndex.GetInt()));
 				
-				LODBoneIndexToWeightIndexMap[LODBoneIndex] = SourceSkeletonBoneIndex.IsValid() ? HTBlendProfile->GetPerBoneInterpolationIndex(SourceSkeletonBoneIndex, TargetSkeleton, Data.Get()) : INDEX_NONE;
+				LODBoneIndexToWeightIndexMap[LODBoneIndex] = SourceSkeletonBoneIndex.IsValid() ? SourceSkeletonBoneIndex.GetInt() : INDEX_NONE;
 			}
 		}
 		else
@@ -469,7 +457,7 @@ void FUAFApplyAdditiveKeyframeWithBlendProfileTask::Execute(UE::UAF::FEvaluation
 			for (int32 LODBoneIndex = 0; LODBoneIndex < NumLODBones; ++LODBoneIndex)
 			{
 				const FSkeletonPoseBoneIndex SkeletonBoneIndex(LODBoneIndexToSkeletonBoneIndexMap[LODBoneIndex]);
-				LODBoneIndexToWeightIndexMap[LODBoneIndex] = SkeletonBoneIndex.IsValid() ? HTBlendProfile->GetPerBoneInterpolationIndex(SkeletonBoneIndex, TargetSkeleton, Data.Get()) : INDEX_NONE;
+				LODBoneIndexToWeightIndexMap[LODBoneIndex] = SkeletonBoneIndex.IsValid() ? SkeletonBoneIndex.GetInt() : INDEX_NONE;
 			}
 		}
 		
@@ -477,24 +465,17 @@ void FUAFApplyAdditiveKeyframeWithBlendProfileTask::Execute(UE::UAF::FEvaluation
 		{
 			check(BaseKeyframe->Pose.GetNumBones() == AdditiveKeyframe->Pose.GetNumBones());
 			
-			TArray<float> BoneWeights;
-			BoneWeights.SetNum(BlendProfileInterface->GetNumBlendEntries());
-			for (int EntryIndex = 0; EntryIndex < BlendProfileInterface->GetNumBlendEntries(); EntryIndex++)
-			{
-				BoneWeights[EntryIndex] = BlendProfileInterface->GetBoneBlendScale(EntryIndex) * AdditiveWeight;
-			}
-			
 			const FTransformArrayView BaseTransformsView = BaseKeyframe->Pose.LocalTransforms.GetView();
 
 			if (AdditiveKeyframe->Pose.IsMeshSpaceAdditive())
 			{
 				BlendWithIdentityAndAccumulatePerBoneWeightMesh(
 					BaseTransformsView, AdditiveKeyframe->Pose.LocalTransforms.GetConstView(),
-					AdditiveKeyframe->Pose.GetLODBoneIndexToParentLODBoneIndexMap(), LODBoneIndexToWeightIndexMap, BoneWeights, AdditiveWeight);
+					AdditiveKeyframe->Pose.GetLODBoneIndexToParentLODBoneIndexMap(), LODBoneIndexToWeightIndexMap, BoneMaskWeights, AdditiveWeight);
 			}
 			else
 			{
-				BlendWithIdentityAndAccumulatePerBoneWeight(BaseTransformsView, AdditiveKeyframe->Pose.LocalTransforms.GetConstView(), LODBoneIndexToWeightIndexMap, BoneWeights, AdditiveWeight);
+				BlendWithIdentityAndAccumulatePerBoneWeight(BaseTransformsView, AdditiveKeyframe->Pose.LocalTransforms.GetConstView(), LODBoneIndexToWeightIndexMap, BoneMaskWeights, AdditiveWeight);
 			}
 
 			NormalizeRotations(BaseTransformsView);
@@ -541,7 +522,7 @@ void FUAFApplyAdditiveKeyframeWithBlendProfileTask::Execute(UE::UAF::FEvaluation
 			
 			// Build up override weight map for masked attributes 
 			TMap<UE::Anim::FAttributeId, float> AttributeMaskTotalWeights;
-			for (const FHierarchyTableBlendProfile::FMaskedAttributeWeight& AttributeWeight : AttributeMaskWeights)
+			for (const FBlendProfileStandaloneCachedData::FMaskedAttributeWeight& AttributeWeight : AttributeMaskWeights)
 			{
 				AttributeMaskTotalWeights.Emplace(AttributeWeight.Attribute, AttributeWeight.Weight * AdditiveWeight);
 			}
diff --git a/Engine/Plugins/Experimental/UAF/UAFAnimGraph/Source/UAFAnimGraph/Private/EvaluationVM/Tasks/BlendKeyframesPerBone.cpp b/Engine/Plugins/Experimental/UAF/UAFAnimGraph/Source/UAFAnimGraph/Private/EvaluationVM/Tasks/BlendKeyframesPerBone.cpp
--- a/Engine/Plugins/Experimental/UAF/UAFAnimGraph/Source/UAFAnimGraph/Private/EvaluationVM/Tasks/BlendKeyframesPerBone.cpp
+++ b/Engine/Plugins/Experimental/UAF/UAFAnimGraph/Source/UAFAnimGraph/Private/EvaluationVM/Tasks/BlendKeyframesPerBone.cpp
@@ -15,6 +15,7 @@
 #include "UAF/ValueRuntime/Transformers/Layer.h"
 #include "UAF/ValueRuntime/Transformers/Overwrite.h"
 #include "UAF/ValueRuntime/Transformers/Sanitize.h"
+#include "UAF/BlendProfile/UAFBlendProfile.h"
 
 #include UE_INLINE_GENERATED_CPP_BY_NAME(BlendKeyframesPerBone)
 
@@ -24,8 +25,7 @@ namespace UE::UAF::Private
 		const FValueBundle& Input,
 		UE::UAF::FEvaluationVM& VM,
 		const FBlendSampleData* BlendData,
-		const IBlendProfileInterface* BlendProfile,
-		bool bPerBoneDataUsesSkeletonIndex)
+		const USkeleton* Skeleton)
 	{
 		const float BlendWeight = BlendData->GetClampedWeight();
 
@@ -49,12 +49,11 @@ namespace UE::UAF::Private
 				const FAttributeTypedSetPtr& BoneSet = BoneMap->GetTypedSet();
 				const int32 NumBonesInSet = BoneMap->Num();
 
-				if (BlendProfile != nullptr)
+				if (Skeleton != nullptr)
 				{
 					const USkeleton* TargetSkeleton = VM.GetActiveEvaluationContext().GetSkeleton();
-					TSharedPtr<IInterpolationIndexProvider::FPerBoneInterpolationData> Data = BlendProfile->GetPerBoneInterpolationData(TargetSkeleton);
 
-					const FSkeletonRemapping& SkeletonRemapping = UE::Anim::FSkeletonRemappingRegistry::Get().GetRemapping(BlendProfile->GetSkeleton(), TargetSkeleton);
+					const FSkeletonRemapping& SkeletonRemapping = UE::Anim::FSkeletonRemappingRegistry::Get().GetRemapping(Skeleton, TargetSkeleton);
 
 					if (SkeletonRemapping.IsValid())
 					{
@@ -65,7 +64,7 @@ namespace UE::UAF::Private
 							const FSkeletonPoseBoneIndex TargetSkeletonBoneIndex = BindingIndex;
 							const FSkeletonPoseBoneIndex SourceSkeletonBoneIndex(SkeletonRemapping.GetSourceSkeletonBoneIndex(TargetSkeletonBoneIndex.GetInt()));
 
-							const int32 WeightIndex = SourceSkeletonBoneIndex.IsValid() ? BlendProfile->GetPerBoneInterpolationIndex(SourceSkeletonBoneIndex, TargetSkeleton, Data.Get()) : INDEX_NONE;
+							const int32 WeightIndex = SourceSkeletonBoneIndex.GetInt();
 							(*SetWeights)[SetIndex].Value = WeightIndex != INDEX_NONE ? BlendData->PerBoneBlendData[WeightIndex] : BlendWeight;
 						}
 					}
@@ -77,30 +76,22 @@ namespace UE::UAF::Private
 
 							const FSkeletonPoseBoneIndex TargetSkeletonBoneIndex = BindingIndex;
 
-							const int32 WeightIndex = BlendProfile->GetPerBoneInterpolationIndex(TargetSkeletonBoneIndex, TargetSkeleton, Data.Get());
+							const int32 WeightIndex = TargetSkeletonBoneIndex.GetInt();
 							(*SetWeights)[SetIndex].Value = WeightIndex != INDEX_NONE ? BlendData->PerBoneBlendData[WeightIndex] : BlendWeight;
 						}
 					}
 				}
 				else
 				{
-					// If no blend profile is provided, we use the provided blend per bone data under the assumption it covers a full pose
-					// either using LODBone indexing or full skeleton indexing
+					// If no skeleton is provided we access the per bone blend data with the assumption that it covers a full pose and without performing any remapping
 					check(BlendData->PerBoneBlendData.Num() >= NumBonesInSet);
 
 					for (FAttributeSetIndex SetIndex(0); SetIndex < NumBonesInSet; ++SetIndex)
 					{
-						if (bPerBoneDataUsesSkeletonIndex)
-						{
-							const FAttributeBindingIndex BindingIndex = BoneSet->GetBindingIndex(SetIndex);
+						const FAttributeBindingIndex BindingIndex = BoneSet->GetBindingIndex(SetIndex);
 
-							const FSkeletonPoseBoneIndex TargetSkeletonBoneIndex = BindingIndex;
-							(*SetWeights)[SetIndex].Value = TargetSkeletonBoneIndex.IsValid() ? BlendData->PerBoneBlendData[TargetSkeletonBoneIndex.GetInt()] : BlendWeight;
-						}
-						else
-						{
-							(*SetWeights)[SetIndex].Value = BlendData->PerBoneBlendData[SetIndex.GetInt()];
-						}
+						const FSkeletonPoseBoneIndex TargetSkeletonBoneIndex = BindingIndex;
+						(*SetWeights)[SetIndex].Value = TargetSkeletonBoneIndex.IsValid() ? BlendData->PerBoneBlendData[TargetSkeletonBoneIndex.GetInt()] : BlendWeight;
 					}
 				}
 			}
@@ -118,10 +109,9 @@ namespace UE::UAF::Private
 		const FValueBundle& Input,
 		UE::UAF::FEvaluationVM& VM,
 		const USkeleton* SourceSkeleton,
-		const TArray<float>& BoneMaskWeights,
-		const IInterpolationIndexProvider& BlendProfile,
-		const UE::Anim::TNamedValueArray<FDefaultAllocator, UE::Anim::FCurveElement>& CurveMaskWeights,
-		const TArray<FAnimNextBlendKeyframePerBoneWithScaleTask::FMaskedAttributeWeight>& AttributeMaskWeights,
+		const UUAFBlendMask::FSkeletonBoneWeightArray* BoneMaskWeights,
+		const UUAFBlendMask::FSkeletonCurveWeightArray* CurveMaskWeights,
+		const UUAFBlendMask::FSkeletonAttributeWeightArray* AttributeMaskWeights,
 		float DefaultWeight)
 	{
 		UScriptStruct* FloatAttributeType = FFloatAnimationAttribute::StaticStruct();
@@ -145,26 +135,38 @@ namespace UE::UAF::Private
 				const FAttributeTypedSetPtr& BoneSet = BoneMap->GetTypedSet();
 				const int32 NumBonesInSet = BoneMap->Num();
 
-				const USkeleton* TargetSkeleton = VM.GetActiveEvaluationContext().GetSkeleton();
-				TSharedPtr<IInterpolationIndexProvider::FPerBoneInterpolationData> Data = BlendProfile.GetPerBoneInterpolationData(TargetSkeleton);
+				bool bUsedSkeletonRemapping = false;
 
-				const FSkeletonRemapping& SkeletonRemapping = UE::Anim::FSkeletonRemappingRegistry::Get().GetRemapping(SourceSkeleton, TargetSkeleton);
-
-				if (SkeletonRemapping.IsValid())
+				if (SourceSkeleton)
 				{
-					for (FAttributeSetIndex SetIndex(0); SetIndex < NumBonesInSet; ++SetIndex)
+					const USkeleton* TargetSkeleton = VM.GetActiveEvaluationContext().GetSkeleton();
+					const FSkeletonRemapping& SkeletonRemapping = UE::Anim::FSkeletonRemappingRegistry::Get().GetRemapping(SourceSkeleton, TargetSkeleton);
+
+					if (SkeletonRemapping.IsValid())
 					{
-						const FAttributeBindingIndex BindingIndex = BoneSet->GetBindingIndex(SetIndex);
+						for (FAttributeSetIndex SetIndex(0); SetIndex < NumBonesInSet; ++SetIndex)
+						{
+							const FAttributeBindingIndex BindingIndex = BoneSet->GetBindingIndex(SetIndex);
 
-						const FSkeletonPoseBoneIndex TargetSkeletonBoneIndex = BindingIndex;
-						const FSkeletonPoseBoneIndex SourceSkeletonBoneIndex(SkeletonRemapping.GetSourceSkeletonBoneIndex(TargetSkeletonBoneIndex.GetInt()));
+							const FSkeletonPoseBoneIndex TargetSkeletonBoneIndex = BindingIndex;
+							const FSkeletonPoseBoneIndex SourceSkeletonBoneIndex(SkeletonRemapping.GetSourceSkeletonBoneIndex(TargetSkeletonBoneIndex.GetInt()));
 
-						const int32 WeightIndex = SourceSkeletonBoneIndex.IsValid() ? BlendProfile.GetPerBoneInterpolationIndex(SourceSkeletonBoneIndex, TargetSkeleton, Data.Get()) : INDEX_NONE;
-						(*SetWeights)[SetIndex].Value = WeightIndex != INDEX_NONE ? BoneMaskWeights[WeightIndex] : 0.0f;
-						(*SetWeights)[SetIndex].Value *= DefaultWeight;
+							if (BoneMaskWeights->IsValidIndex(SourceSkeletonBoneIndex.GetInt()))
+							{
+								(*SetWeights)[SetIndex].Value = SourceSkeletonBoneIndex.IsValid() ? (*BoneMaskWeights)[SourceSkeletonBoneIndex.GetInt()] : 0.0f;
+							}
+							else
+							{
+								(*SetWeights)[SetIndex].Value = 0.0f;
+							}
+							(*SetWeights)[SetIndex].Value *= DefaultWeight;
+						}
+
+						bUsedSkeletonRemapping = true;
 					}
 				}
-				else
+
+				if (!bUsedSkeletonRemapping)
 				{
 					for (FAttributeSetIndex SetIndex(0); SetIndex < NumBonesInSet; ++SetIndex)
 					{
@@ -172,8 +174,14 @@ namespace UE::UAF::Private
 
 						const FSkeletonPoseBoneIndex TargetSkeletonBoneIndex = BindingIndex;
 
-						const int32 WeightIndex = BlendProfile.GetPerBoneInterpolationIndex(TargetSkeletonBoneIndex, TargetSkeleton, Data.Get());
-						(*SetWeights)[SetIndex].Value = WeightIndex != INDEX_NONE ? BoneMaskWeights[WeightIndex] : 0.0f;
+						if (BoneMaskWeights->IsValidIndex(TargetSkeletonBoneIndex.GetInt()))
+						{
+							(*SetWeights)[SetIndex].Value = TargetSkeletonBoneIndex.IsValid() ? (*BoneMaskWeights)[TargetSkeletonBoneIndex.GetInt()] : 0.0f;
+						}
+						else
+						{
+							(*SetWeights)[SetIndex].Value = 0.0f;
+						}
 						(*SetWeights)[SetIndex].Value *= DefaultWeight;
 					}
 				}
@@ -183,7 +191,7 @@ namespace UE::UAF::Private
 				SetWeights->FillWithValue(FFloatAnimationAttribute{ 0.0f });
 
 				const FAttributeTypedSetPtr& FloatSet = FloatMap->GetTypedSet();
-				CurveMaskWeights.ForEachElement([&FloatSet, SetWeights, DefaultWeight](const UE::Anim::FCurveElement& Curve)
+				CurveMaskWeights->ForEachElement([&FloatSet, SetWeights, DefaultWeight](const UE::Anim::FCurveElement& Curve)
 					{
 						const FAttributeSetIndex SetIndex = FloatSet->FindIndex(Curve.Name);
 						if (SetIndex.IsValid())
@@ -193,7 +201,7 @@ namespace UE::UAF::Private
 						}
 					});
 
-				for (const FAnimNextBlendKeyframePerBoneWithScaleTask::FMaskedAttributeWeight& AttributeWeight : AttributeMaskWeights)
+				for (const FBlendProfileStandaloneCachedData::FMaskedAttributeWeight& AttributeWeight : *AttributeMaskWeights)
 				{
 					const FAttributeSetIndex SetIndex = FloatSet->FindIndex(AttributeWeight.Attribute.GetName());
 					if (SetIndex.IsValid())
@@ -208,7 +216,7 @@ namespace UE::UAF::Private
 				SetWeights->FillWithValue(FFloatAnimationAttribute{ 0.0f });
 
 				const FAttributeTypedSetPtr& TypedSet = It.GetMap()->GetTypedSet();
-				for (const FAnimNextBlendKeyframePerBoneWithScaleTask::FMaskedAttributeWeight& AttributeWeight : AttributeMaskWeights)
+				for (const FBlendProfileStandaloneCachedData::FMaskedAttributeWeight& AttributeWeight : *AttributeMaskWeights)
 				{
 					const FAttributeSetIndex SetIndex = TypedSet->FindIndex(AttributeWeight.Attribute.GetName());
 					if (SetIndex.IsValid())
@@ -223,23 +231,12 @@ namespace UE::UAF::Private
 	}
 }
 
-FAnimNextBlendOverwriteKeyframePerBoneWithScaleTask FAnimNextBlendOverwriteKeyframePerBoneWithScaleTask::Make(const IBlendProfileInterface* BlendProfile, const FBlendSampleData& BlendData, float ScaleFactor)
+FAnimNextBlendOverwriteKeyframePerBoneWithScaleTask FAnimNextBlendOverwriteKeyframePerBoneWithScaleTask::Make(const USkeleton* Skeleton, const FBlendSampleData& BlendData, float ScaleFactor)
 {
 	FAnimNextBlendOverwriteKeyframePerBoneWithScaleTask Task;
-	Task.BlendProfile = BlendProfile;
-	Task.BlendData = &BlendData;
-	Task.ScaleFactor = ScaleFactor;
-
-	return Task;
-}
-
-FAnimNextBlendOverwriteKeyframePerBoneWithScaleTask FAnimNextBlendOverwriteKeyframePerBoneWithScaleTask::Make(const FBlendSampleData& BlendData, float ScaleFactor, bool bPerBoneDataUsesSkeletonIndex)
-{
-	FAnimNextBlendOverwriteKeyframePerBoneWithScaleTask Task;
-	Task.BlendProfile = nullptr;
+	Task.SourceSkeleton = Skeleton;
 	Task.BlendData = &BlendData;
 	Task.ScaleFactor = ScaleFactor;
-	Task.bPerBoneDataUsesSkeletonIndex = bPerBoneDataUsesSkeletonIndex;
 
 	return Task;
 }
@@ -248,7 +245,7 @@ void FAnimNextBlendOverwriteKeyframePerBoneWithScaleTask::Execute(UE::UAF::FEval
 {
 	using namespace UE::UAF;
 
-	if (BlendProfile == nullptr && BlendData->PerBoneBlendData.Num() == 0)
+	if (BlendData->PerBoneBlendData.Num() == 0)
 	{
 		// If we don't have a blend profile and per bone blend data, we blend the whole pose
 		Super::Execute(VM);
@@ -266,7 +263,7 @@ void FAnimNextBlendOverwriteKeyframePerBoneWithScaleTask::Execute(UE::UAF::FEval
 
 		const float BlendWeight = BlendData->GetClampedWeight();
 
-		const FValueBundleStack PerValueWeights = Private::BuildPerValueWeights(*Input, VM, BlendData, BlendProfile, bPerBoneDataUsesSkeletonIndex);
+		const FValueBundleStack PerValueWeights = Private::BuildPerValueWeights(*Input, VM, BlendData, SourceSkeleton);
 
 		// Modify our input in place
 		Transformers::FOverwrite::Apply(VM.GetTransformerMap(), *Input, PerValueWeights, BlendWeight, *Input);
@@ -288,47 +285,25 @@ void FAnimNextBlendOverwriteKeyframePerBoneWithScaleTask::Execute(UE::UAF::FEval
 		TArray<int32> LODBoneIndexToWeightIndexMap;
 		LODBoneIndexToWeightIndexMap.AddUninitialized(NumLODBones);
 
-		if (BlendProfile)
-		{
-			const USkeleton* TargetSkeleton = Keyframe->Pose.GetSkeletonAsset();
-			TSharedPtr<IInterpolationIndexProvider::FPerBoneInterpolationData> Data = BlendProfile->GetPerBoneInterpolationData(TargetSkeleton);
+		const USkeleton* TargetSkeleton = Keyframe->Pose.GetSkeletonAsset();
 
-			const FSkeletonRemapping& SkeletonRemapping = UE::Anim::FSkeletonRemappingRegistry::Get().GetRemapping(BlendProfile->GetSkeleton(), TargetSkeleton);
+		const FSkeletonRemapping& SkeletonRemapping = UE::Anim::FSkeletonRemappingRegistry::Get().GetRemapping(SourceSkeleton, TargetSkeleton);
 
-			if (SkeletonRemapping.IsValid())
-			{
-				for (int32 LODBoneIndex = 0; LODBoneIndex < NumLODBones; ++LODBoneIndex)
-				{
-					const FSkeletonPoseBoneIndex TargetSkeletonBoneIndex(LODBoneIndexToSkeletonBoneIndexMap[LODBoneIndex]);
-					const FSkeletonPoseBoneIndex SourceSkeletonBoneIndex(SkeletonRemapping.GetSourceSkeletonBoneIndex(TargetSkeletonBoneIndex.GetInt()));
-					LODBoneIndexToWeightIndexMap[LODBoneIndex] = SourceSkeletonBoneIndex.IsValid() ? BlendProfile->GetPerBoneInterpolationIndex(SourceSkeletonBoneIndex, TargetSkeleton, Data.Get()) : INDEX_NONE;
-				}
-			}
-			else
+		if (SkeletonRemapping.IsValid())
+		{
+			for (int32 LODBoneIndex = 0; LODBoneIndex < NumLODBones; ++LODBoneIndex)
 			{
-				for (int32 LODBoneIndex = 0; LODBoneIndex < NumLODBones; ++LODBoneIndex)
-				{
-					const FSkeletonPoseBoneIndex SkeletonBoneIndex(LODBoneIndexToSkeletonBoneIndexMap[LODBoneIndex]);
-					LODBoneIndexToWeightIndexMap[LODBoneIndex] = BlendProfile->GetPerBoneInterpolationIndex(SkeletonBoneIndex, TargetSkeleton, Data.Get());
-				}
+				const FSkeletonPoseBoneIndex TargetSkeletonBoneIndex(LODBoneIndexToSkeletonBoneIndexMap[LODBoneIndex]);
+				const FSkeletonPoseBoneIndex SourceSkeletonBoneIndex(SkeletonRemapping.GetSourceSkeletonBoneIndex(TargetSkeletonBoneIndex.GetInt()));
+				LODBoneIndexToWeightIndexMap[LODBoneIndex] = SourceSkeletonBoneIndex.GetInt();
 			}
 		}
 		else
 		{
-			// If no blend profile is provided, we use the provided blend per bone data under the assumption it covers a full pose
-			// either using LODBone indexing or full skeleton indexing
-			check(BlendData->PerBoneBlendData.Num() >= NumLODBones);
 			for (int32 LODBoneIndex = 0; LODBoneIndex < NumLODBones; ++LODBoneIndex)
 			{
-				if (bPerBoneDataUsesSkeletonIndex)
-				{
-					const FSkeletonPoseBoneIndex TargetSkeletonBoneIndex(LODBoneIndexToSkeletonBoneIndexMap[LODBoneIndex]);
-					LODBoneIndexToWeightIndexMap[LODBoneIndex] = TargetSkeletonBoneIndex.GetInt();
-				}
-				else
-				{
-					LODBoneIndexToWeightIndexMap[LODBoneIndex] = LODBoneIndex;
-				}
+				const FSkeletonPoseBoneIndex SkeletonBoneIndex(LODBoneIndexToSkeletonBoneIndexMap[LODBoneIndex]);
+				LODBoneIndexToWeightIndexMap[LODBoneIndex] = SkeletonBoneIndex.GetInt();
 			}
 		}
 
@@ -361,25 +336,13 @@ void FAnimNextBlendOverwriteKeyframePerBoneWithScaleTask::Execute(UE::UAF::FEval
 	}
 }
 
-FAnimNextBlendAddKeyframePerBoneWithScaleTask FAnimNextBlendAddKeyframePerBoneWithScaleTask::Make(const IBlendProfileInterface* BlendProfile, const FBlendSampleData& BlendDataA, const FBlendSampleData& BlendDataB, float ScaleFactor)
+FAnimNextBlendAddKeyframePerBoneWithScaleTask FAnimNextBlendAddKeyframePerBoneWithScaleTask::Make(const USkeleton* Skeleton, const FBlendSampleData& BlendDataA, const FBlendSampleData& BlendDataB, float ScaleFactor)
 {
 	FAnimNextBlendAddKeyframePerBoneWithScaleTask Task;
-	Task.BlendProfile = BlendProfile;
-	Task.BlendDataA = &BlendDataA;
-	Task.BlendDataB = &BlendDataB;
-	Task.ScaleFactor = ScaleFactor;
-
-	return Task;
-}
-
-FAnimNextBlendAddKeyframePerBoneWithScaleTask FAnimNextBlendAddKeyframePerBoneWithScaleTask::Make(const FBlendSampleData& BlendDataA, const FBlendSampleData& BlendDataB, float ScaleFactor, bool bPerBoneDataUsesSkeletonIndex)
-{
-	FAnimNextBlendAddKeyframePerBoneWithScaleTask Task;
-	Task.BlendProfile = nullptr;
+	Task.SourceSkeleton = Skeleton;
 	Task.BlendDataA = &BlendDataA;
 	Task.BlendDataB = &BlendDataB;
 	Task.ScaleFactor = ScaleFactor;
-	Task.bPerBoneDataUsesSkeletonIndex = bPerBoneDataUsesSkeletonIndex;
 
 	return Task;
 }
@@ -388,9 +351,8 @@ void FAnimNextBlendAddKeyframePerBoneWithScaleTask::Execute(UE::UAF::FEvaluation
 {
 	using namespace UE::UAF;
 
-	if (BlendProfile == nullptr
-		&& ((BlendDataA->PerBoneBlendData.Num() == 0 || BlendDataB->PerBoneBlendData.Num() == 0)
-			|| (BlendDataA->PerBoneBlendData.Num() != BlendDataB->PerBoneBlendData.Num())))
+	if ((BlendDataA->PerBoneBlendData.Num() == 0 || BlendDataB->PerBoneBlendData.Num() == 0)
+		|| (BlendDataA->PerBoneBlendData.Num() != BlendDataB->PerBoneBlendData.Num()))
 	{
 		// If we don't have a blend profile, we blend the whole pose
 		Super::Execute(VM);
@@ -416,7 +378,7 @@ void FAnimNextBlendAddKeyframePerBoneWithScaleTask::Execute(UE::UAF::FEvaluation
 
 		const float BlendWeight = BlendDataA->GetClampedWeight();
 
-		const FValueBundleStack PerValueWeightsA = Private::BuildPerValueWeights(*InputA, VM, BlendDataA, BlendProfile, bPerBoneDataUsesSkeletonIndex);
+		const FValueBundleStack PerValueWeightsA = Private::BuildPerValueWeights(*InputA, VM, BlendDataA, SourceSkeleton);
 
 		// Modify input B in place
 		// We flip A and B here because state is pushed in reversed order
@@ -449,12 +411,12 @@ void FAnimNextBlendAddKeyframePerBoneWithScaleTask::Execute(UE::UAF::FEvaluation
 		TArray<int32> LODBoneIndexToWeightIndexMap;
 		LODBoneIndexToWeightIndexMap.AddUninitialized(NumLODBones);
 
-		if (BlendProfile)
+		bool bUsedSkeletonRemapping = false;
+
+		if (SourceSkeleton)
 		{
 			const USkeleton* TargetSkeleton = KeyframeB->Pose.GetSkeletonAsset();
-			TSharedPtr<IInterpolationIndexProvider::FPerBoneInterpolationData> Data = BlendProfile->GetPerBoneInterpolationData(TargetSkeleton);
-
-			const FSkeletonRemapping& SkeletonRemapping = UE::Anim::FSkeletonRemappingRegistry::Get().GetRemapping(BlendProfile->GetSkeleton(), TargetSkeleton);
+			const FSkeletonRemapping& SkeletonRemapping = UE::Anim::FSkeletonRemappingRegistry::Get().GetRemapping(SourceSkeleton, TargetSkeleton);
 
 			if (SkeletonRemapping.IsValid())
 			{
@@ -462,34 +424,19 @@ void FAnimNextBlendAddKeyframePerBoneWithScaleTask::Execute(UE::UAF::FEvaluation
 				{
 					const FSkeletonPoseBoneIndex TargetSkeletonBoneIndex(LODBoneIndexToSkeletonBoneIndexMap[LODBoneIndex]);
 					const FSkeletonPoseBoneIndex SourceSkeletonBoneIndex(SkeletonRemapping.GetSourceSkeletonBoneIndex(TargetSkeletonBoneIndex.GetInt()));
-					LODBoneIndexToWeightIndexMap[LODBoneIndex] = SourceSkeletonBoneIndex.IsValid() ? BlendProfile->GetPerBoneInterpolationIndex(SourceSkeletonBoneIndex, TargetSkeleton, Data.Get()) : INDEX_NONE;
-				}
-			}
-			else
-			{
-				for (int32 LODBoneIndex = 0; LODBoneIndex < NumLODBones; ++LODBoneIndex)
-				{
-					const FSkeletonPoseBoneIndex SkeletonBoneIndex(LODBoneIndexToSkeletonBoneIndexMap[LODBoneIndex]);
-					LODBoneIndexToWeightIndexMap[LODBoneIndex] = BlendProfile->GetPerBoneInterpolationIndex(SkeletonBoneIndex, TargetSkeleton, Data.Get());
+					LODBoneIndexToWeightIndexMap[LODBoneIndex] = SourceSkeletonBoneIndex.GetInt();
 				}
+
+				bUsedSkeletonRemapping = true;
 			}
 		}
-		else
+
+		if (!bUsedSkeletonRemapping)
 		{
-			// If no blend profile is provided, we assume the provided bone data encompasses a full pose to be blended together 
-			// either using LOD based indexing or full skeleton indexing
-			check(BlendDataA->PerBoneBlendData.Num() >= NumLODBones);
 			for (int32 LODBoneIndex = 0; LODBoneIndex < NumLODBones; ++LODBoneIndex)
 			{
-				if (bPerBoneDataUsesSkeletonIndex)
-				{
-					const FSkeletonPoseBoneIndex TargetSkeletonBoneIndex(LODBoneIndexToSkeletonBoneIndexMap[LODBoneIndex]);
-					LODBoneIndexToWeightIndexMap[LODBoneIndex] = TargetSkeletonBoneIndex.GetInt();
-				}
-				else
-				{
-					LODBoneIndexToWeightIndexMap[LODBoneIndex] = LODBoneIndex;
-				}
+				const FSkeletonPoseBoneIndex SkeletonBoneIndex(LODBoneIndexToSkeletonBoneIndexMap[LODBoneIndex]);
+				LODBoneIndexToWeightIndexMap[LODBoneIndex] = SkeletonBoneIndex.GetInt();
 			}
 		}
 
@@ -527,14 +474,32 @@ void FAnimNextBlendAddKeyframePerBoneWithScaleTask::Execute(UE::UAF::FEvaluation
 	}
 }
 
-FAnimNextBlendKeyframePerBoneWithScaleTask FAnimNextBlendKeyframePerBoneWithScaleTask::Make(TSharedPtr<IInterpolationIndexProvider> BlendProfile, const USkeleton* Skeleton, const TArray<float>& BoneMaskWeights, const UE::Anim::TNamedValueArray<FDefaultAllocator, UE::Anim::FCurveElement>& CurveMaskWeights, const TArray<FMaskedAttributeWeight>& AttributeMaskWeights, float ScaleFactor)
+FAnimNextBlendKeyframePerBoneWithScaleTask FAnimNextBlendKeyframePerBoneWithScaleTask::Make(
+	const UUAFBlendMask* BlendMask,
+	float ScaleFactor)
+{
+	FAnimNextBlendKeyframePerBoneWithScaleTask Task;
+	Task.SourceSkeleton = BlendMask->GetSkeleton();
+	Task.BoneMaskWeights = &BlendMask->GetSkeletonBoneWeights();
+	Task.CurveMaskWeights = &BlendMask->GetSkeletonCurveWeights();
+	Task.AttributeMaskWeights = &BlendMask->GetSkeletonAttributeWeights();
+	Task.ScaleFactor = ScaleFactor;
+
+	return Task;
+}
+
+FAnimNextBlendKeyframePerBoneWithScaleTask FAnimNextBlendKeyframePerBoneWithScaleTask::Make(
+	const USkeleton* Skeleton,
+	const UUAFBlendMask::FSkeletonBoneWeightArray* InBoneWeights,
+	const UUAFBlendMask::FSkeletonCurveWeightArray* InCurveWeights,
+	const UUAFBlendMask::FSkeletonAttributeWeightArray* InAttributeWeights,
+	float ScaleFactor)
 {
 	FAnimNextBlendKeyframePerBoneWithScaleTask Task;
-	Task.BlendProfile = BlendProfile;
 	Task.SourceSkeleton = Skeleton;
-	Task.BoneMaskWeights = BoneMaskWeights;
-	Task.CurveMaskWeights = CurveMaskWeights;
-	Task.AttributeMaskWeights = AttributeMaskWeights;
+	Task.BoneMaskWeights = InBoneWeights;
+	Task.CurveMaskWeights = InCurveWeights;
+	Task.AttributeMaskWeights = InAttributeWeights;
 	Task.ScaleFactor = ScaleFactor;
 
 	return Task;
@@ -544,7 +509,7 @@ void FAnimNextBlendKeyframePerBoneWithScaleTask::Execute(UE::UAF::FEvaluationVM&
 {
 	using namespace UE::UAF;
 
-	if (BlendProfile == nullptr)
+	if (!BoneMaskWeights || BoneMaskWeights->IsEmpty())
 	{
 		// If we don't have a blend profile, we blend the whole pose
 		Super::Execute(VM);
@@ -568,7 +533,7 @@ void FAnimNextBlendKeyframePerBoneWithScaleTask::Execute(UE::UAF::FEvaluationVM&
 			return;
 		}
 
-		const FValueBundleStack PerValueWeightsLayer = Private::BuildPerValueWeights(*Base, VM, SourceSkeleton, BoneMaskWeights, *BlendProfile, CurveMaskWeights, AttributeMaskWeights, ScaleFactor);
+		const FValueBundleStack PerValueWeightsLayer = Private::BuildPerValueWeights(*Base, VM, SourceSkeleton, BoneMaskWeights, CurveMaskWeights, AttributeMaskWeights, ScaleFactor);
 
 		// Apply our layer onto our base and modify it in-place
 		Transformers::FLayer::Apply(VM.GetTransformerMap(), *Base, *Layer, PerValueWeightsLayer, ScaleFactor, *Base);
@@ -596,7 +561,6 @@ void FAnimNextBlendKeyframePerBoneWithScaleTask::Execute(UE::UAF::FEvaluationVM&
 		}
 
 		const USkeleton* TargetSkeleton = KeyframeB->Pose.GetSkeletonAsset();
-		TSharedPtr<IInterpolationIndexProvider::FPerBoneInterpolationData> Data = BlendProfile->GetPerBoneInterpolationData(TargetSkeleton);
 
 		const FSkeletonRemapping& SkeletonRemapping = UE::Anim::FSkeletonRemappingRegistry::Get().GetRemapping(SourceSkeleton, TargetSkeleton);
 
@@ -612,7 +576,7 @@ void FAnimNextBlendKeyframePerBoneWithScaleTask::Execute(UE::UAF::FEvaluationVM&
 			{
 				const FSkeletonPoseBoneIndex TargetSkeletonBoneIndex(LODBoneIndexToSkeletonBoneIndexMap[LODBoneIndex]);
 				const FSkeletonPoseBoneIndex SourceSkeletonBoneIndex(SkeletonRemapping.GetSourceSkeletonBoneIndex(TargetSkeletonBoneIndex.GetInt()));
-				LODBoneIndexToWeightIndexMap[LODBoneIndex] = SourceSkeletonBoneIndex.IsValid() ? BlendProfile->GetPerBoneInterpolationIndex(SourceSkeletonBoneIndex, TargetSkeleton, Data.Get()) : INDEX_NONE;
+				LODBoneIndexToWeightIndexMap[LODBoneIndex] = SourceSkeletonBoneIndex.IsValid() ? SourceSkeletonBoneIndex.GetInt() : INDEX_NONE;
 			}
 		}
 		else
@@ -620,7 +584,7 @@ void FAnimNextBlendKeyframePerBoneWithScaleTask::Execute(UE::UAF::FEvaluationVM&
 			for (int32 LODBoneIndex = 0; LODBoneIndex < NumLODBones; ++LODBoneIndex)
 			{
 				const FSkeletonPoseBoneIndex SkeletonBoneIndex(LODBoneIndexToSkeletonBoneIndexMap[LODBoneIndex]);
-				LODBoneIndexToWeightIndexMap[LODBoneIndex] = BlendProfile->GetPerBoneInterpolationIndex(SkeletonBoneIndex, TargetSkeleton, Data.Get());
+				LODBoneIndexToWeightIndexMap[LODBoneIndex] = SkeletonBoneIndex.GetInt();
 			}
 		}
 
@@ -630,11 +594,11 @@ void FAnimNextBlendKeyframePerBoneWithScaleTask::Execute(UE::UAF::FEvaluationVM&
 
 			BlendOverwritePerBoneWithScale(
 				KeyframeA->Pose.LocalTransforms.GetView(), KeyframeA->Pose.LocalTransforms.GetConstView(),
-				LODBoneIndexToWeightIndexMap, BoneMaskWeights, 1.0f, true);
+				LODBoneIndexToWeightIndexMap, *BoneMaskWeights, 1.0f, true);
 
 			BlendAddPerBoneWithScale(
 				KeyframeA->Pose.LocalTransforms.GetView(), KeyframeB->Pose.LocalTransforms.GetConstView(),
-				LODBoneIndexToWeightIndexMap, BoneMaskWeights, ScaleFactor);
+				LODBoneIndexToWeightIndexMap, *BoneMaskWeights, ScaleFactor);
 		}
 
 		if (EnumHasAnyFlags(VM.GetFlags(), EEvaluationFlags::Curves))
@@ -646,28 +610,31 @@ void FAnimNextBlendKeyframePerBoneWithScaleTask::Execute(UE::UAF::FEvaluationVM&
 			{
 				FBlendedCurve FilteredCurves;
 
-				// Multiply per-curve blend weights by matching blend pose curves
-				UE::Anim::FNamedValueArrayUtils::Intersection(
-					TargetCurve,
-					CurveMaskWeights,
-					[this, &FilteredCurves](const UE::Anim::FCurveElement& InBlendElement, const UE::Anim::FCurveElement& InMaskElement) mutable
-					{
-						FilteredCurves.Add(InBlendElement.Name, InBlendElement.Value * InMaskElement.Value);
-					});
+				if (CurveMaskWeights)
+				{
+					// Multiply per-curve blend weights by matching blend pose curves
+					UE::Anim::FNamedValueArrayUtils::Intersection(
+						TargetCurve,
+						*CurveMaskWeights,
+						[this, &FilteredCurves](const UE::Anim::FCurveElement& InBlendElement, const UE::Anim::FCurveElement& InMaskElement) mutable
+						{
+							FilteredCurves.Add(InBlendElement.Name, InBlendElement.Value * InMaskElement.Value);
+						});
 
-				// Override blend curve values with premultipled curves
-				TargetCurve.Combine(FilteredCurves);
+					// Override blend curve values with premultipled curves
+					TargetCurve.Combine(FilteredCurves);
 
-				// Remove curves that have been filtered by the mask, curves with no mask value defined remain, even with 0.0 value
-				UE::Anim::FNamedValueArrayUtils::RemoveByPredicate(
-					TargetCurve,
-					CurveMaskWeights,
-					[](const UE::Anim::FCurveElement& InBaseElement, const UE::Anim::FCurveElement& InMaskElement)
-					{
-						const bool bKeep = InMaskElement.Value == 0.0;
-						return bKeep;
-					}
-				);
+					// Remove curves that have been filtered by the mask, curves with no mask value defined remain, even with 0.0 value
+					UE::Anim::FNamedValueArrayUtils::RemoveByPredicate(
+						TargetCurve,
+						*CurveMaskWeights,
+						[](const UE::Anim::FCurveElement& InBaseElement, const UE::Anim::FCurveElement& InMaskElement)
+						{
+							const bool bKeep = InMaskElement.Value == 0.0;
+							return bKeep;
+						}
+					);
+				}
 
 				// Combine base and filtered pre-multiplied blend curves
 				UE::Anim::FNamedValueArrayUtils::Union(
@@ -724,68 +691,71 @@ void FAnimNextBlendKeyframePerBoneWithScaleTask::Execute(UE::UAF::FEvaluationVM&
 			// -    | b     | 0.0    | -
 
 
-			// 1. Blend attributes according to the bone blend weights,
-			// i.e. an attribute's weight is determined by the weight of its attached bone.
-			UE::Anim::Attributes::BlendAttributesPerBone(BaseAttributes, BlendAttributes, BoneMaskWeights, OutputAttributes);
-
-
-			// 2. For each attribute that has a custom weight, i.e. one's that shouldn't be weighted by its attached bone, go and correct the blended value.
-			for (const FMaskedAttributeWeight& MaskedAttribute : AttributeMaskWeights)
+			if (AttributeMaskWeights)
 			{
-				const TConstArrayView<TWeakObjectPtr<UScriptStruct>> UniqueTypes = OutputAttributes.GetUniqueTypes();
+				// 1. Blend attributes according to the bone blend weights,
+				// i.e. an attribute's weight is determined by the weight of its attached bone.
+				UE::Anim::Attributes::BlendAttributesPerBone(BaseAttributes, BlendAttributes, *BoneMaskWeights, OutputAttributes);
 
-				for (int32 UniqueTypeIndex = 0; UniqueTypeIndex < UniqueTypes.Num(); ++UniqueTypeIndex)
-				{
-					const TConstArrayView<FAttributeId> AttributeKeys = OutputAttributes.GetKeys(UniqueTypeIndex);
-					const TConstArrayView<TWrappedAttribute<FAnimStackAllocator>> AttributeValues = OutputAttributes.GetValues(UniqueTypeIndex);
-
-					const TWeakObjectPtr<UScriptStruct> AttributeType = UniqueTypes[UniqueTypeIndex];
-					const IAttributeBlendOperator* Operator = UE::Anim::AttributeTypes::GetTypeOperator(AttributeType);
 
-					TWrappedAttribute<FAnimStackAllocator> DefaultData(AttributeType.Get());
-					AttributeType->InitializeStruct(DefaultData.GetPtr<void>());
+				// 2. For each attribute that has a custom weight, i.e. one's that shouldn't be weighted by its attached bone, go and correct the blended value.
+				for (const FBlendProfileStandaloneCachedData::FMaskedAttributeWeight& MaskedAttribute : *AttributeMaskWeights)
+				{
+					const TConstArrayView<TWeakObjectPtr<UScriptStruct>> UniqueTypes = OutputAttributes.GetUniqueTypes();
 
-					uint8* OutputData = OutputAttributes.Find(AttributeType.Get(), MaskedAttribute.Attribute);
-					if (OutputData)
+					for (int32 UniqueTypeIndex = 0; UniqueTypeIndex < UniqueTypes.Num(); ++UniqueTypeIndex)
 					{
-						uint8* BaseData = BaseAttributes.Find(AttributeType.Get(), MaskedAttribute.Attribute);
-						uint8* BlendData = BlendAttributes.Find(AttributeType.Get(), MaskedAttribute.Attribute);
+						const TConstArrayView<FAttributeId> AttributeKeys = OutputAttributes.GetKeys(UniqueTypeIndex);
+						const TConstArrayView<TWrappedAttribute<FAnimStackAllocator>> AttributeValues = OutputAttributes.GetValues(UniqueTypeIndex);
 
-						if (BaseData && BlendData)
-						{
-							// Base | Blend | Weight | Output
-							// ------------------------------
-							// a    | b     | 1.0    | b
-							// a    | b     | k      | lerp(a, b, k)
-							// a    | b     | 0.0    | a
+						const TWeakObjectPtr<UScriptStruct> AttributeType = UniqueTypes[UniqueTypeIndex];
+						const IAttributeBlendOperator* Operator = UE::Anim::AttributeTypes::GetTypeOperator(AttributeType);
 
-							Operator->Interpolate(BaseData, BlendData, MaskedAttribute.Weight, OutputData);
-						}
-						else if (BaseData && !BlendData)
-						{
-							// Base | Blend | Weight | Output
-							// ------------------------------
-							// a    | -     | 1.0    | a
-							// a    | -     | k      | a
-							// a    | -     | 0.0    | a
+						TWrappedAttribute<FAnimStackAllocator> DefaultData(AttributeType.Get());
+						AttributeType->InitializeStruct(DefaultData.GetPtr<void>());
 
-							*OutputData = *BaseData;
-						}
-						else if (!BaseData && BlendData)
+						uint8* OutputData = OutputAttributes.Find(AttributeType.Get(), MaskedAttribute.Attribute);
+						if (OutputData)
 						{
-							// Base | Blend | Weight | Output
-							// ------------------------------
-							// -    | b     | 1.0    | b
-							// -    | b     | k      | lerp(default, b, k)
-							// -    | b     | 0.0    | -
+							uint8* BaseData = BaseAttributes.Find(AttributeType.Get(), MaskedAttribute.Attribute);
+							uint8* BlendData = BlendAttributes.Find(AttributeType.Get(), MaskedAttribute.Attribute);
+
+							if (BaseData && BlendData)
+							{
+								// Base | Blend | Weight | Output
+								// ------------------------------
+								// a    | b     | 1.0    | b
+								// a    | b     | k      | lerp(a, b, k)
+								// a    | b     | 0.0    | a
 
-							if (MaskedAttribute.Weight != 0.0f)
+								Operator->Interpolate(BaseData, BlendData, MaskedAttribute.Weight, OutputData);
+							}
+							else if (BaseData && !BlendData)
 							{
-								Operator->Interpolate(DefaultData.GetPtr<void>(), BlendData, MaskedAttribute.Weight, OutputData);
+								// Base | Blend | Weight | Output
+								// ------------------------------
+								// a    | -     | 1.0    | a
+								// a    | -     | k      | a
+								// a    | -     | 0.0    | a
+
+								*OutputData = *BaseData;
 							}
-							else
+							else if (!BaseData && BlendData)
 							{
-								OutputAttributes.Remove(AttributeType.Get(), MaskedAttribute.Attribute);
+								// Base | Blend | Weight | Output
+								// ------------------------------
+								// -    | b     | 1.0    | b
+								// -    | b     | k      | lerp(default, b, k)
+								// -    | b     | 0.0    | -
+
+								if (MaskedAttribute.Weight != 0.0f)
+								{
+									Operator->Interpolate(DefaultData.GetPtr<void>(), BlendData, MaskedAttribute.Weight, OutputData);
+								}
+								else
+								{
+									OutputAttributes.Remove(AttributeType.Get(), MaskedAttribute.Attribute);
+								}
 							}
 						}
 					}
diff --git a/Engine/Plugins/Experimental/UAF/UAFAnimGraph/Source/UAFAnimGraph/Private/Injection/MontageTrait.cpp b/Engine/Plugins/Experimental/UAF/UAFAnimGraph/Source/UAFAnimGraph/Private/Injection/MontageTrait.cpp
--- a/Engine/Plugins/Experimental/UAF/UAFAnimGraph/Source/UAFAnimGraph/Private/Injection/MontageTrait.cpp
+++ b/Engine/Plugins/Experimental/UAF/UAFAnimGraph/Source/UAFAnimGraph/Private/Injection/MontageTrait.cpp
@@ -34,6 +34,7 @@ namespace UE::UAF
 		GeneratorMacro(IUpdateTraversal) \
 		GeneratorMacro(ITimeline) \
 		GeneratorMacro(IGroupSynchronization) \
+		GeneratorMacro(IGarbageCollection) \
 
 	GENERATE_ANIM_TRAIT_IMPLEMENTATION(FMontageTrait, TRAIT_INTERFACE_ENUMERATOR, NULL_ANIM_TRAIT_INTERFACE_ENUMERATOR, NULL_ANIM_TRAIT_EVENT_ENUMERATOR)
 	#undef TRAIT_INTERFACE_ENUMERATOR
@@ -71,6 +72,7 @@ namespace UE::UAF
 		AdditiveEvalTaskDataForSlot.Reset();
 		BlendDataForBlendProfiles.Reset();
 		AdditiveBlendDataForBlendProfiles.Reset();
+		SkeletonDataForBlendProfiles.Reset();
 		bReturnPreUpdateState = true;
 		SyncCandidateIndex = INDEX_NONE;
 	}
@@ -173,7 +175,9 @@ namespace UE::UAF
 									SampleBlendData.PerBoneBlendData.AddZeroed(NumBones);
 
 									const UBlendProfile* BlendProfile = EvalData.ActiveBlendProfile;
-									for (int32 PerBoneIndex = 0; PerBoneIndex < SampleBlendData.PerBoneBlendData.Num(); ++PerBoneIndex)
+									InstanceData->SkeletonDataForBlendProfiles.Add(BlendProfile ? BlendProfile->GetSkeleton() : nullptr);
+
+;									for (int32 PerBoneIndex = 0; PerBoneIndex < SampleBlendData.PerBoneBlendData.Num(); ++PerBoneIndex)
 									{
 										// If we have a active blend profile we calculate the bone weight according to its settings
 										// Otherwise we fall back to use the pose weight as bone weight
@@ -292,6 +296,7 @@ namespace UE::UAF
 			// Add the pose weight to blend in the source with
 			if (InstanceData->bUsesBlendProfiles)
 			{
+				InstanceData->SkeletonDataForBlendProfiles.Add(nullptr);
 				FBlendSampleData& SourceSampleBlendData = InstanceData->BlendDataForBlendProfiles.AddDefaulted_GetRef();
 				SourceSampleBlendData.TotalWeight = InstanceData->SourceWeight;
 				SourceSampleBlendData.PerBoneBlendData.SetNumZeroed(NumBones);
@@ -399,13 +404,17 @@ namespace UE::UAF
 			if (InstanceData->BlendDataForBlendProfiles.Num() > 0)
 			{
 				const FBlendSampleData& SampleBlendData = InstanceData->BlendDataForBlendProfiles[0];
-				Context.AppendTask(FAnimNextBlendOverwriteKeyframePerBoneWithScaleTask::Make(SampleBlendData, SampleBlendData.GetClampedWeight(), true));
+				const USkeleton* BlendProfileSkeleton = InstanceData->SkeletonDataForBlendProfiles[0];
+
+				Context.AppendTask(FAnimNextBlendOverwriteKeyframePerBoneWithScaleTask::Make(BlendProfileSkeleton, SampleBlendData, SampleBlendData.GetClampedWeight()));
 
 				for (int32 BlendDataIndex = 1; BlendDataIndex < InstanceData->BlendDataForBlendProfiles.Num(); ++BlendDataIndex)
 				{
 					const FBlendSampleData& SampleBlendDataA = InstanceData->BlendDataForBlendProfiles[BlendDataIndex];
 					const FBlendSampleData& SampleBlendDataB = InstanceData->BlendDataForBlendProfiles[BlendDataIndex - 1];
-					Context.AppendTask(FAnimNextBlendAddKeyframePerBoneWithScaleTask::Make(SampleBlendDataA, SampleBlendDataB, SampleBlendDataA.GetClampedWeight(), true));
+					const USkeleton* BlendProfileSkeletonA = InstanceData->SkeletonDataForBlendProfiles[BlendDataIndex];
+
+					Context.AppendTask(FAnimNextBlendAddKeyframePerBoneWithScaleTask::Make(BlendProfileSkeletonA, SampleBlendDataA, SampleBlendDataB, SampleBlendDataA.GetClampedWeight()));
 				}
 
 				// Ensure that rotations are normalized after blending 
@@ -544,5 +553,13 @@ namespace UE::UAF
 		const FInstanceData* InstanceData = Binding.GetInstanceData<FInstanceData>();
 		return InstanceData->SyncGroupParams;
 	}
+
+	void FMontageTrait::AddReferencedObjects(const FExecutionContext& Context, const TTraitBinding<IGarbageCollection>& Binding, FReferenceCollector& Collector) const
+	{
+		IGarbageCollection::AddReferencedObjects(Context, Binding, Collector);
+
+		FInstanceData* InstanceData = Binding.GetInstanceData<FInstanceData>();
+		Collector.AddReferencedObjects(InstanceData->SkeletonDataForBlendProfiles);
+	}
 }
 
diff --git a/Engine/Plugins/Experimental/UAF/UAFAnimGraph/Source/UAFAnimGraph/Private/TraitInterfaces/ISmoothBlendPerBone.cpp b/Engine/Plugins/Experimental/UAF/UAFAnimGraph/Source/UAFAnimGraph/Private/TraitInterfaces/ISmoothBlendPerBone.cpp
--- a/Engine/Plugins/Experimental/UAF/UAFAnimGraph/Source/UAFAnimGraph/Private/TraitInterfaces/ISmoothBlendPerBone.cpp
+++ b/Engine/Plugins/Experimental/UAF/UAFAnimGraph/Source/UAFAnimGraph/Private/TraitInterfaces/ISmoothBlendPerBone.cpp
@@ -21,7 +21,7 @@ namespace UE::UAF
 	}
 #endif // WITH_EDITOR
 
-	TSharedPtr<const IBlendProfileInterface> ISmoothBlendPerBone::GetBlendProfile(FExecutionContext& Context, const TTraitBinding<ISmoothBlendPerBone>& Binding, int32 ChildIndex) const
+	const UUAFBlendProfile* ISmoothBlendPerBone::GetBlendProfile(FExecutionContext& Context, const TTraitBinding<ISmoothBlendPerBone>& Binding, int32 ChildIndex) const
 	{
 		TTraitBinding<ISmoothBlendPerBone> SuperBinding;
 		if (Binding.GetStackInterfaceSuper(SuperBinding))
diff --git a/Engine/Plugins/Experimental/UAF/UAFAnimGraph/Source/UAFAnimGraph/Private/Traits/ApplyAdditivePerBoneTrait.cpp b/Engine/Plugins/Experimental/UAF/UAFAnimGraph/Source/UAFAnimGraph/Private/Traits/ApplyAdditivePerBoneTrait.cpp
--- a/Engine/Plugins/Experimental/UAF/UAFAnimGraph/Source/UAFAnimGraph/Private/Traits/ApplyAdditivePerBoneTrait.cpp
+++ b/Engine/Plugins/Experimental/UAF/UAFAnimGraph/Source/UAFAnimGraph/Private/Traits/ApplyAdditivePerBoneTrait.cpp
@@ -3,11 +3,11 @@
 
 #include "Traits/ApplyAdditivePerBoneTrait.h"
 
-#include "HierarchyTable.h"
 #include "EvaluationVM/Tasks/ApplyAdditiveKeyframePerBone.h"
 #include "EvaluationVM/Tasks/ConvertRotationsToLocalSpace.h"
 #include "EvaluationVM/Tasks/ConvertRotationsToMeshSpace.h"
 #include "EvaluationVM/Tasks/NormalizeRotations.h"
+#include "UAF/BlendMask/UAFBlendMask.h"
 
 namespace UE::UAF
 {
@@ -33,7 +33,7 @@ void FApplyAdditivePerBoneTrait::PreUpdate(FUpdateTraversalContext& Context, con
 	TTraitBinding<IContinuousBlend> ContinuousBlendTrait;
 	Binding.GetStackInterface(ContinuousBlendTrait);
 	
-	const TObjectPtr<UHierarchyTable> BlendProfile = SharedData->GetBlendProfile(Binding);
+	const TObjectPtr<UUAFBlendMask> BlendMask = SharedData->GetBlendMask(Binding);
 	if (!InstanceData->Base.IsValid())
 	{
 		InstanceData->Base = Context.AllocateNodeInstance(Binding, SharedData->Base);
@@ -44,7 +44,7 @@ void FApplyAdditivePerBoneTrait::PreUpdate(FUpdateTraversalContext& Context, con
 	}
 	
 	const float AdditiveWeight = ContinuousBlendTrait.GetBlendWeight(Context, 1);
-	if (FAnimWeight::IsRelevant(AdditiveWeight) && BlendProfile)
+	if (FAnimWeight::IsRelevant(AdditiveWeight) && BlendMask)
 	{
 		if (!InstanceData->Additive.IsValid())
 		{
@@ -55,13 +55,6 @@ void FApplyAdditivePerBoneTrait::PreUpdate(FUpdateTraversalContext& Context, con
 		{
 			InstanceData->bWasAdditiveRelevant = true;
 		}
-		
-		// Check if the blend profile has changed since the last update and update relevant instance data accordingly
-		if (!InstanceData->CachedBlendProfile.IsValid() || InstanceData->CachedBlendProfile.Get() != BlendProfile)
-		{
-			InstanceData->CachedBlendProfile = BlendProfile;
-			InstanceData->HierarchyTableBlendProfile = MakeShared<FHierarchyTableBlendProfile>(BlendProfile, EBlendProfileMode::BlendMask);
-		}
 	}
 	else
 	{
@@ -76,7 +69,9 @@ void FApplyAdditivePerBoneTrait::PostEvaluate(FEvaluateTraversalContext& Context
 	const FSharedData* SharedData = Binding.GetSharedData<FSharedData>();
 	const FInstanceData* InstanceData = Binding.GetInstanceData<FInstanceData>();
 	
-	if (InstanceData->Base.IsValid() && InstanceData->Additive.IsValid() && InstanceData->HierarchyTableBlendProfile.IsValid())
+	const UUAFBlendMask* BlendMask = SharedData->GetBlendMask(Binding);
+
+	if (InstanceData->Base.IsValid() && InstanceData->Additive.IsValid() && BlendMask)
 	{
 		TTraitBinding<IContinuousBlend> ContinuousBlendTrait;
 		Binding.GetStackInterface(ContinuousBlendTrait);
@@ -89,7 +84,7 @@ void FApplyAdditivePerBoneTrait::PostEvaluate(FEvaluateTraversalContext& Context
 			Context.AppendTask(FAnimNextConvertRotationsToMeshSpaceTask::Make(2));
 		}
 
-		Context.AppendTask(FUAFApplyAdditiveKeyframeWithBlendProfileTask::Make(AdditiveWeight, InstanceData->HierarchyTableBlendProfile));
+		Context.AppendTask(FUAFApplyAdditiveKeyframeWithBlendMaskTask::Make(AdditiveWeight, BlendMask));
 		Context.AppendTask(FAnimNextNormalizeKeyframeRotationsTask());
 
 		if (bBlendInMeshSpace)
@@ -109,7 +104,7 @@ void FApplyAdditivePerBoneTrait::QueueChildrenForTraversal(FUpdateTraversalConte
 	const FSharedData* SharedData = Binding.GetSharedData<FSharedData>();
 	const FInstanceData* InstanceData = Binding.GetInstanceData<FInstanceData>();
 	
-	const TObjectPtr<UHierarchyTable> BlendProfile = SharedData->GetBlendProfile(Binding);
+	const UUAFBlendMask* BlendMask = SharedData->GetBlendMask(Binding);
 	
 	TTraitBinding<IContinuousBlend> ContinuousBlendTrait;
 	Binding.GetStackInterface(ContinuousBlendTrait);
@@ -121,7 +116,7 @@ void FApplyAdditivePerBoneTrait::QueueChildrenForTraversal(FUpdateTraversalConte
 		TraversalQueue.Push(InstanceData->Base, TraitState.WithWeight(BaseWeight).AsNewlyRelevant(!InstanceData->bWasBaseRelevant));
 	}
 
-	if (InstanceData->Additive.IsValid() && FAnimWeight::IsRelevant(AdditiveWeight) && BlendProfile)
+	if (InstanceData->Additive.IsValid() && FAnimWeight::IsRelevant(AdditiveWeight) && BlendMask)
 	{
 		TraversalQueue.Push(InstanceData->Additive, TraitState.WithWeight(AdditiveWeight).AsNewlyRelevant(!InstanceData->bWasAdditiveRelevant));
 	}
@@ -146,7 +141,7 @@ float FApplyAdditivePerBoneTrait::GetBlendWeight(const FExecutionContext& Contex
 {
 	const FSharedData* SharedData = Binding.GetSharedData<FSharedData>();
 	float Alpha = FMath::Clamp(SharedData->GetAlpha(Binding), 0.0f, 1.0f);
-	if (SharedData->GetBlendProfile(Binding) == nullptr)
+	if (SharedData->GetBlendMask(Binding) == nullptr)
 	{
 		Alpha = 0.0f;
 	}
@@ -169,10 +164,12 @@ void FApplyAdditivePerBoneTrait::AddReferencedObjects(const FExecutionContext& C
 {
 	IGarbageCollection::AddReferencedObjects(Context, Binding, Collector);
 	
-	FInstanceData* InstanceData = Binding.GetInstanceData<FInstanceData>();
-	if (InstanceData->CachedBlendProfile.IsValid())
+	const FSharedData* SharedData = Binding.GetSharedData<FSharedData>();
+	TObjectPtr<UUAFBlendMask> BlendMask = SharedData->GetBlendMask(Binding);
+
+	if (BlendMask)
 	{
-		Collector.AddReferencedObject(InstanceData->CachedBlendProfile);
+		Collector.AddReferencedObject(BlendMask);
 	}
 }
 	
diff --git a/Engine/Plugins/Experimental/UAF/UAFAnimGraph/Source/UAFAnimGraph/Private/Traits/ApplyAdditivePerBoneTrait.h b/Engine/Plugins/Experimental/UAF/UAFAnimGraph/Source/UAFAnimGraph/Private/Traits/ApplyAdditivePerBoneTrait.h
--- a/Engine/Plugins/Experimental/UAF/UAFAnimGraph/Source/UAFAnimGraph/Private/Traits/ApplyAdditivePerBoneTrait.h
+++ b/Engine/Plugins/Experimental/UAF/UAFAnimGraph/Source/UAFAnimGraph/Private/Traits/ApplyAdditivePerBoneTrait.h
@@ -31,9 +31,6 @@ namespace UE::UAF
 
 			bool bWasBaseRelevant = false;
 			bool bWasAdditiveRelevant = false;
-			
-			TSharedPtr<const FHierarchyTableBlendProfile> HierarchyTableBlendProfile = nullptr;
-			TWeakObjectPtr<UHierarchyTable> CachedBlendProfile = nullptr;
 		};
 
 
diff --git a/Engine/Plugins/Experimental/UAF/UAFAnimGraph/Source/UAFAnimGraph/Private/Traits/BlendLayer.cpp b/Engine/Plugins/Experimental/UAF/UAFAnimGraph/Source/UAFAnimGraph/Private/Traits/BlendLayer.cpp
--- a/Engine/Plugins/Experimental/UAF/UAFAnimGraph/Source/UAFAnimGraph/Private/Traits/BlendLayer.cpp
+++ b/Engine/Plugins/Experimental/UAF/UAFAnimGraph/Source/UAFAnimGraph/Private/Traits/BlendLayer.cpp
@@ -12,38 +12,7 @@
 #include "HierarchyTable.h"
 #include "MaskProfile/HierarchyTableTypeMask.h"
 #include "SkeletonHierarchyTableType.h"
-
-
-class FHierarchyTableBlendProfileProvider : public IInterpolationIndexProvider
-{
-public:
-	TObjectPtr<UHierarchyTable> HierarchyTable;
-
-	// IInterpolationIndexProvider
-	virtual TSharedPtr<FPerBoneInterpolationData> GetPerBoneInterpolationData(const USkeleton* Skeleton) const
-	{
-		return nullptr;
-	}
-
-	virtual int32 GetPerBoneInterpolationIndex(const FCompactPoseBoneIndex& InCompactPoseBoneIndex, const FBoneContainer& BoneContainer, const IInterpolationIndexProvider::FPerBoneInterpolationData* Data) const override
-	{
-		return 0;
-	}
-
-	virtual int32 GetPerBoneInterpolationIndex(const FSkeletonPoseBoneIndex InSkeletonBoneIndex, const USkeleton* TargetSkeleton, const IInterpolationIndexProvider::FPerBoneInterpolationData* Data) const override
-	{
-		if (HierarchyTable)
-		{
-			if (HierarchyTable->GetTableData().IsValidIndex(InSkeletonBoneIndex.GetInt()))
-			{
-				return InSkeletonBoneIndex.GetInt();
-			}
-		}
-		return INDEX_NONE;
-	}
-	// End IInterpolationIndexProvider
-};
-
+#include "UAF/BlendMask/UAFBlendMask.h"
 
 namespace UE::UAF
 {
@@ -64,20 +33,16 @@ namespace UE::UAF
 		const FSharedData* SharedData = Binding.GetSharedData<FSharedData>();
 		const FInstanceData* InstanceData = Binding.GetInstanceData<FInstanceData>();
 
-		TObjectPtr<UHierarchyTable> HierarchyTable = SharedData->BlendProfile;
+		const UUAFBlendMask* BlendMask = SharedData->GetBlendMask(Binding);
 
-		if (InstanceData->ChildBase.IsValid() && InstanceData->ChildBlend.IsValid() && HierarchyTable)
+		if (InstanceData->ChildBase.IsValid() && InstanceData->ChildBlend.IsValid() && BlendMask)
 		{
 			// We have two children, interpolate them
 
 			TTraitBinding<IContinuousBlend> ContinuousBlendTrait;
 			Binding.GetStackInterface(ContinuousBlendTrait);
 
-			TSharedPtr<FHierarchyTableBlendProfileProvider> InterpolationIndexProvider = MakeShared<FHierarchyTableBlendProfileProvider>();
-			InterpolationIndexProvider->HierarchyTable = HierarchyTable;
-
-			const FHierarchyTable_TableType_Skeleton& TableMetadata = HierarchyTable->GetTableMetadata<FHierarchyTable_TableType_Skeleton>();
-			const TObjectPtr<USkeleton> Skeleton = TableMetadata.Skeleton;
+			const USkeleton* Skeleton = BlendMask->GetSkeleton();
 
 			const float BlendWeight = ContinuousBlendTrait.GetBlendWeight(Context, 1);
 			const bool bBlendInMeshSpace = SharedData->GetbBlendInMeshSpace(Binding);
@@ -87,7 +52,7 @@ namespace UE::UAF
 				Context.AppendTask(FAnimNextConvertRotationsToMeshSpaceTask::Make(2));
 			}
 
-			Context.AppendTask(FAnimNextBlendKeyframePerBoneWithScaleTask::Make(InterpolationIndexProvider, Skeleton, InstanceData->BoneMaskWeights, InstanceData->CurveMaskWeights, InstanceData->AttributeMaskWeights, BlendWeight));
+			Context.AppendTask(FAnimNextBlendKeyframePerBoneWithScaleTask::Make(BlendMask, BlendWeight));
 			Context.AppendTask(FAnimNextNormalizeKeyframeRotationsTask());
 
 			if (bBlendInMeshSpace)
@@ -109,68 +74,7 @@ namespace UE::UAF
 		TTraitBinding<IContinuousBlend> ContinuousBlendTrait;
 		Binding.GetStackInterface(ContinuousBlendTrait);
 
-		TObjectPtr<UHierarchyTable> HierarchyTable = SharedData->BlendProfile;
-
-		if (InstanceData->bBoneMaskWeightsNeedEvaluating)
-		{
-			InstanceData->BoneMaskWeights.Empty();
-
-			if (HierarchyTable)
-			{
-				const int32 NumBlendEntries = SharedData->BlendProfile->GetTableData().Num();
-				InstanceData->BoneMaskWeights.SetNum(NumBlendEntries);
-
-				for (int32 Index = 0; Index < NumBlendEntries; ++Index)
-				{
-					InstanceData->BoneMaskWeights[Index] = HierarchyTable->GetTableEntry(Index)->GetValue<FHierarchyTable_ElementType_Mask>()->Value;
-				}
-			}
-
-			InstanceData->bBoneMaskWeightsNeedEvaluating = false;
-		}
-
-		if (InstanceData->bCurveMaskWeightsNeedEvaluating)
-		{
-			InstanceData->CurveMaskWeights.Empty();
-
-			if (HierarchyTable)
-			{
-				for (const FHierarchyTableEntryData& TableEntry : HierarchyTable->GetTableData())
-				{
-					if (TableEntry.GetMetadata<FHierarchyTable_TablePayloadType_Skeleton>().EntryType == ESkeletonHierarchyTable_TablePayloadEntryType::Curve)
-					{
-						const float EntryWeight = TableEntry.GetValue<FHierarchyTable_ElementType_Mask>()->Value;
-						InstanceData->CurveMaskWeights.Add<UE::Anim::FCurveElement>({ TableEntry.Identifier, EntryWeight });
-					}
-				}
-			}
-
-			InstanceData->bCurveMaskWeightsNeedEvaluating = false;
-		}
-
-		if (InstanceData->bAttributeMaskWeightsNeedEvaluating)
-		{
-			InstanceData->AttributeMaskWeights.Empty();
-
-			if (HierarchyTable)
-			{
-				for (const FHierarchyTableEntryData& TableEntry : HierarchyTable->GetTableData())
-				{
-					if (TableEntry.GetMetadata<FHierarchyTable_TablePayloadType_Skeleton>().EntryType == ESkeletonHierarchyTable_TablePayloadEntryType::Attribute)
-					{
-						const FHierarchyTableEntryData* ParentEntry = HierarchyTable->GetTableEntry(TableEntry.Parent);
-						check(ParentEntry && ParentEntry->GetMetadata<FHierarchyTable_TablePayloadType_Skeleton>().EntryType == ESkeletonHierarchyTable_TablePayloadEntryType::Bone);
-
-						UE::Anim::FAttributeId Attribute(TableEntry.Identifier, TableEntry.Parent, TEXT("bone"));
-
-						const float EntryWeight = TableEntry.GetValue<FHierarchyTable_ElementType_Mask>()->Value;
-						InstanceData->AttributeMaskWeights.Add(FAnimNextBlendKeyframePerBoneWithScaleTask::FMaskedAttributeWeight(Attribute, EntryWeight));
-					}
-				}
-			}
-
-			InstanceData->bAttributeMaskWeightsNeedEvaluating = false;
-		}
+		const UUAFBlendMask* BlendMask = SharedData->GetBlendMask(Binding);
 
 		if (!InstanceData->ChildBase.IsValid())
 		{
@@ -182,7 +86,7 @@ namespace UE::UAF
 		}
 		
 		const float ChildWeight = ContinuousBlendTrait.GetBlendWeight(Context, 1);
-		if ((FAnimWeight::IsRelevant(ChildWeight) || SharedData->GetbAlwaysUpdateChildBlend(Binding)) && SharedData->BlendProfile)
+		if ((FAnimWeight::IsRelevant(ChildWeight) || SharedData->GetbAlwaysUpdateChildBlend(Binding)) && BlendMask)
 		{
 			if (!InstanceData->ChildBlend.IsValid())
 			{
@@ -214,7 +118,7 @@ namespace UE::UAF
 			TraversalQueue.Push(InstanceData->ChildBase, TraitState.WithWeight(1.0f).AsNewlyRelevant(!InstanceData->bWasChildBaseRelevant));
 		}
 
-		if (InstanceData->ChildBlend.IsValid() && SharedData->BlendProfile)
+		if (InstanceData->ChildBlend.IsValid() && SharedData->GetBlendMask(Binding))
 		{
 			TraversalQueue.Push(InstanceData->ChildBlend, TraitState.WithWeight(BlendWeightBlend).AsNewlyRelevant(!InstanceData->bWasChildBlendRelevant));
 		}
@@ -241,7 +145,7 @@ namespace UE::UAF
 
 		const float BlendWeight = SharedData->GetBlendWeight(Binding);
 		float ClampedWeight = FMath::Clamp(BlendWeight, 0.0f, 1.0f);
-		if (!SharedData->BlendProfile)
+		if (!SharedData->GetBlendMask(Binding))
 		{
 			ClampedWeight = 0.0f;
 		}
diff --git a/Engine/Plugins/Experimental/UAF/UAFAnimGraph/Source/UAFAnimGraph/Private/Traits/BlendLayer.h b/Engine/Plugins/Experimental/UAF/UAFAnimGraph/Source/UAFAnimGraph/Private/Traits/BlendLayer.h
--- a/Engine/Plugins/Experimental/UAF/UAFAnimGraph/Source/UAFAnimGraph/Private/Traits/BlendLayer.h
+++ b/Engine/Plugins/Experimental/UAF/UAFAnimGraph/Source/UAFAnimGraph/Private/Traits/BlendLayer.h
@@ -33,16 +33,6 @@ namespace UE::UAF
 
 			bool bWasChildBaseRelevant = false;
 			bool bWasChildBlendRelevant = false;
-
-			bool bBoneMaskWeightsNeedEvaluating = true;
-			bool bCurveMaskWeightsNeedEvaluating = true;
-			bool bAttributeMaskWeightsNeedEvaluating = true;
-
-			TArray<float> BoneMaskWeights;
-
-			UE::Anim::TNamedValueArray<FDefaultAllocator, UE::Anim::FCurveElement> CurveMaskWeights;
-
-			TArray<FAnimNextBlendKeyframePerBoneWithScaleTask::FMaskedAttributeWeight> AttributeMaskWeights;
 		};
 
 
diff --git a/Engine/Plugins/Experimental/UAF/UAFAnimGraph/Source/UAFAnimGraph/Private/Traits/BlendSmootherPerBone.cpp b/Engine/Plugins/Experimental/UAF/UAFAnimGraph/Source/UAFAnimGraph/Private/Traits/BlendSmootherPerBone.cpp
--- a/Engine/Plugins/Experimental/UAF/UAFAnimGraph/Source/UAFAnimGraph/Private/Traits/BlendSmootherPerBone.cpp
+++ b/Engine/Plugins/Experimental/UAF/UAFAnimGraph/Source/UAFAnimGraph/Private/Traits/BlendSmootherPerBone.cpp
@@ -9,14 +9,14 @@
 #include "TraitInterfaces/IHierarchy.h"
 #include "EvaluationVM/Tasks/BlendKeyframesPerBone.h"
 #include "EvaluationVM/Tasks/NormalizeRotations.h"
+#include "UAF/BlendProfile/UAFBlendProfile.h"
 
 #include UE_INLINE_GENERATED_CPP_BY_NAME(BlendSmootherPerBone)
 
 namespace UE::UAF
 {
 	AUTO_REGISTER_ANIM_TRAIT(FBlendSmootherPerBoneCoreTrait)
-	AUTO_REGISTER_ANIM_TRAIT(FBlendProfilePerChildProviderTrait)
-	AUTO_REGISTER_ANIM_TRAIT(FBlendProfileProviderTrait)
+	AUTO_REGISTER_ANIM_TRAIT(FBlendSmootherPerBoneTrait)
 
 	// Trait implementation boilerplate
 	#define TRAIT_INTERFACE_ENUMERATOR(GeneratorMacro) \
@@ -36,18 +36,8 @@ namespace UE::UAF
 	// Trait implementation boilerplate
 	#define TRAIT_INTERFACE_ENUMERATOR(GeneratorMacro) \
 		GeneratorMacro(ISmoothBlendPerBone) \
-		GeneratorMacro(IUpdate) \
 
-	GENERATE_ANIM_TRAIT_IMPLEMENTATION(FBlendProfilePerChildProviderTrait, TRAIT_INTERFACE_ENUMERATOR, NULL_ANIM_TRAIT_INTERFACE_ENUMERATOR, NULL_ANIM_TRAIT_EVENT_ENUMERATOR)
-	#undef TRAIT_INTERFACE_ENUMERATOR
-	#undef TRAIT_REQUIRED_INTERFACE_ENUMERATOR
-
-		// Trait implementation boilerplate
-	#define TRAIT_INTERFACE_ENUMERATOR(GeneratorMacro) \
-		GeneratorMacro(ISmoothBlendPerBone) \
-		GeneratorMacro(IUpdate) \
-
-		GENERATE_ANIM_TRAIT_IMPLEMENTATION(FBlendProfileProviderTrait, TRAIT_INTERFACE_ENUMERATOR, NULL_ANIM_TRAIT_INTERFACE_ENUMERATOR, NULL_ANIM_TRAIT_EVENT_ENUMERATOR)
+	GENERATE_ANIM_TRAIT_IMPLEMENTATION(FBlendSmootherPerBoneTrait, TRAIT_INTERFACE_ENUMERATOR, NULL_ANIM_TRAIT_INTERFACE_ENUMERATOR, NULL_ANIM_TRAIT_EVENT_ENUMERATOR)
 	#undef TRAIT_INTERFACE_ENUMERATOR
 	#undef TRAIT_REQUIRED_INTERFACE_ENUMERATOR
 
@@ -55,20 +45,11 @@ namespace UE::UAF
 	{
 		const FInstanceData* InstanceData = Binding.GetInstanceData<FInstanceData>();
 
-		TTraitBinding<ISmoothBlend> SmoothBlendTrait;
-		Binding.GetStackInterface(SmoothBlendTrait);
-
-		TTraitBinding<IDiscreteBlend> DiscreteBlendTrait;
-		Binding.GetStackInterface(DiscreteBlendTrait);
-
-		const int32 NumChildren = InstanceData->PerChildBlendData.Num();
-
 		int32 NumBlending = 0;
-		for (int32 ChildIndex = 0; ChildIndex < NumChildren; ++ChildIndex)
-		{
-			const FBlendData& ChildBlendData = InstanceData->PerChildBlendData[ChildIndex];
 
-			NumBlending += DiscreteBlendTrait.GetBlendWeight(Context, ChildIndex) > 0.0f ? 1 : 0;
+		for (const FBlendData& ChildBlendData : InstanceData->PerChildBlendData)
+		{
+			NumBlending += ChildBlendData.bIsBlending ? 1 : 0;
 		}
 
 		if (NumBlending < 2)
@@ -76,30 +57,38 @@ namespace UE::UAF
 			return;	// If we don't have at least 2 children blending, there is nothing to do
 		}
 
-		int32 ChildIndex = NumChildren - 1;
+		TTraitBinding<IDiscreteBlend> DiscreteBlendTrait;
+		Binding.GetStackInterface(DiscreteBlendTrait);
+
+		// Children are visited depth first, in the order returned
+		// As such, when we evaluate the task program, the keyframe of the last child will be
+		// on top of the keyframe stack
+		// We thus process children in reverse order
+
+		// The last child override the top keyframe and scales it
+		int32 ChildIndex = InstanceData->PerChildBlendData.Num() - 1;
 		for (; ChildIndex >= 0; --ChildIndex)
 		{
-			const float Weight = DiscreteBlendTrait.GetBlendWeight(Context, ChildIndex);
-
-			if (Weight <= 0.0f)
+			const FBlendData& ChildBlendData = InstanceData->PerChildBlendData[ChildIndex];
+			if (!ChildBlendData.bIsBlending)
 			{
-				continue;	// Skip inactive child
+				continue;	// Skip this inactive child
 			}
 
-			const FBlendData& ChildBlendData = InstanceData->PerChildBlendData[ChildIndex];
-			const FBlendSampleData& PoseSampleData = InstanceData->PerBoneSampleData[ChildIndex];
-
-			TSharedPtr<const IBlendProfileInterface> BlendProfileInterface = ChildBlendData.BlendProfileInterface;
+			const float ChildWeight = DiscreteBlendTrait.GetBlendWeight(Context, ChildIndex);
 
-			if (BlendProfileInterface)
+			// This trait controls the blend weight and owns it
+			if (InstanceData->BlendProfileSkeleton)
 			{
-				Context.AppendTask(FAnimNextBlendOverwriteKeyframePerBoneWithScaleTask::Make(BlendProfileInterface.Get(), PoseSampleData, Weight));
+				const FBlendSampleData& SampleData = InstanceData->PerBoneSampleData[ChildIndex];
+				Context.AppendTask(FAnimNextBlendOverwriteKeyframePerBoneWithScaleTask::Make(InstanceData->BlendProfileSkeleton, SampleData, 1.0f));
 			}
 			else
 			{
-				Context.AppendTask(FAnimNextBlendOverwriteKeyframeWithScaleTask::Make(Weight));
+				Context.AppendTask(FAnimNextBlendOverwriteKeyframeWithScaleTask::Make(ChildWeight));
 			}
 
+			// We found the last child to blend
 			break;
 		}
 
@@ -107,29 +96,33 @@ namespace UE::UAF
 		ChildIndex--;
 		for (; ChildIndex >= 0; --ChildIndex)
 		{
-			const float Weight = DiscreteBlendTrait.GetBlendWeight(Context, ChildIndex);
-
 			const FBlendData& ChildBlendData = InstanceData->PerChildBlendData[ChildIndex];
-			if (Weight <= 0.0f)
+			if (!ChildBlendData.bIsBlending)
 			{
-				continue;	// Skip inactive child
+				continue;	// Skip this inactive child
 			}
+			
+			const float ChildWeight = DiscreteBlendTrait.GetBlendWeight(Context, ChildIndex);
 
-			const FBlendSampleData& PoseSampleDataA = InstanceData->PerBoneSampleData[ChildIndex];
-			const FBlendSampleData& PoseSampleDataB = InstanceData->PerBoneSampleData[ChildIndex + 1];	// Above on the keyframe stack
-
-			TSharedPtr<const IBlendProfileInterface> BlendProfileInterface = ChildBlendData.BlendProfileInterface;
-
-			if (BlendProfileInterface)
+			// This trait controls the blend weight and owns it
+			if (InstanceData->BlendProfileSkeleton)
 			{
-				Context.AppendTask(FAnimNextBlendAddKeyframePerBoneWithScaleTask::Make(BlendProfileInterface.Get(), PoseSampleDataA, PoseSampleDataB, Weight));
+				const FBlendSampleData& PoseSampleDataA = InstanceData->PerBoneSampleData[ChildIndex];
+				const FBlendSampleData& PoseSampleDataB = InstanceData->PerBoneSampleData[ChildIndex + 1];
+
+				Context.AppendTask(FAnimNextBlendAddKeyframePerBoneWithScaleTask::Make(
+					InstanceData->BlendProfileSkeleton,
+					PoseSampleDataA,
+					PoseSampleDataB,
+					1.0f));
 			}
 			else
 			{
-				Context.AppendTask(FAnimNextBlendAddKeyframeWithScaleTask::Make(Weight));
+				Context.AppendTask(FAnimNextBlendAddKeyframeWithScaleTask::Make(ChildWeight));
 			}
 		}
 
+		// Once we are done, we normalize rotations
 		Context.AppendTask(FAnimNextNormalizeKeyframeRotationsTask());
 	}
 
@@ -138,6 +131,20 @@ namespace UE::UAF
 		const FSharedData* SharedData = Binding.GetSharedData<FSharedData>();
 		FInstanceData* InstanceData = Binding.GetInstanceData<FInstanceData>();
 
+		TTraitBinding<IDiscreteBlend> DiscreteBlendTrait;
+		Binding.GetStackInterface(DiscreteBlendTrait);
+
+		{
+			const int32 NumChildren = InstanceData->PerChildBlendData.Num();
+			check(InstanceData->ChildOrdering.Num() == NumChildren);
+			check(InstanceData->PerBoneSampleData.Num() == NumChildren);		
+		
+			for (int32 ChildIndex = 0; ChildIndex < NumChildren; ++ChildIndex)
+			{
+				InstanceData->PerChildBlendData[ChildIndex].OldWeight = DiscreteBlendTrait.GetBlendWeight(Context, ChildIndex);
+			}
+		}
+
 		// If this is our first update, allocate our blend data
 		if (InstanceData->PerChildBlendData.IsEmpty())
 		{
@@ -145,42 +152,70 @@ namespace UE::UAF
 		}
 
 		// Update the traits below us, they might trigger a transition
+		// ISmoothBlend does most of the bookkeeping here
 		IUpdate::PreUpdate(Context, Binding, TraitState);
 
-		TTraitBinding<IDiscreteBlend> DiscreteBlendTrait;
-		Binding.GetStackInterface(DiscreteBlendTrait);
-
-		const int32 DestinationChildIndex = DiscreteBlendTrait.GetBlendDestinationChildIndex(Context);
-
 		const int32 NumChildren = InstanceData->PerChildBlendData.Num();
-
-		// If we're using a blend profile, extract the scales and build blend sample data
-		for (int32 ChildIndex = 0; ChildIndex < NumChildren; ++ChildIndex)
+		for (int32 ChildOrderIndex = 0; ChildOrderIndex < NumChildren; ++ChildOrderIndex)
 		{
-			const float BlendWeight = DiscreteBlendTrait.GetBlendWeight(Context, ChildIndex);
-			const FAlphaBlend* BlendState = DiscreteBlendTrait.GetBlendState(Context, ChildIndex);
-
-			FBlendSampleData& PoseSampleData = InstanceData->PerBoneSampleData[ChildIndex];
-			PoseSampleData.TotalWeight = BlendWeight;
+			const int32 ChildIndex = InstanceData->ChildOrdering[ChildOrderIndex];
+			const float ChildWeight = DiscreteBlendTrait.GetBlendWeight(Context, ChildIndex);
 
-			const FBlendData& BlendData = InstanceData->PerChildBlendData[ChildIndex];
+			FBlendSampleData& ChildSampleData = InstanceData->PerBoneSampleData[ChildIndex];
+			const FBlendData& ChildBlendData = InstanceData->PerChildBlendData[ChildIndex];
 
-			for (int32 PerBoneIndex = 0; PerBoneIndex < PoseSampleData.PerBoneBlendData.Num(); ++PerBoneIndex)
+			if (ChildBlendData.BlendProfile && ChildOrderIndex != 0)
+			{
+				const UUAFBlendProfile::FSkeletonBoneWeightArray& BoneProfileWeights = ChildBlendData.BlendProfile->GetSkeletonBoneWeights();
+
+				if (ChildBlendData.BlendProfile->GetType() == EUAFBlendProfileType::WeightFactor)
+				{
+					for (int32 PerBoneIndex = 0; PerBoneIndex < ChildSampleData.PerBoneBlendData.Num(); ++PerBoneIndex)
+					{
+						const float ProfileValue = BoneProfileWeights[PerBoneIndex];
+						const float Weight = (ProfileValue > ZERO_ANIMWEIGHT_THRESH) ? ChildWeight / ProfileValue : 1.0f;
+						ChildSampleData.PerBoneBlendData[PerBoneIndex] = FMath::Max(Weight, ZERO_ANIMWEIGHT_THRESH);
+					}
+				}
+				else
+				{
+					const FAlphaBlend* BlendState = DiscreteBlendTrait.GetBlendState(Context, ChildIndex);
+
+					for (int32 PerBoneIndex = 0; PerBoneIndex < ChildSampleData.PerBoneBlendData.Num(); ++PerBoneIndex)
+					{
+						const float ProfileValue = BoneProfileWeights[PerBoneIndex];
+
+						ChildSampleData.PerBoneBlendData[PerBoneIndex] = UBlendProfile::CalculateBoneWeight(
+							ProfileValue,
+							EBlendProfileMode::TimeFactor,
+							*BlendState,
+							ChildBlendData.OldWeight,
+							ChildWeight,
+							false);
+					}
+				}
+			}
+			else if (InstanceData->BlendProfileSkeleton)
 			{
-				PoseSampleData.PerBoneBlendData[PerBoneIndex] = UBlendProfile::CalculateBoneWeight(
-					BlendData.BlendProfileInterface->GetBoneBlendScale(PerBoneIndex),
-					EBlendProfileMode::TimeFactor, *BlendState, BlendData.StartAlpha, BlendWeight, false /*bInverse*/);
+				for (int32 PerBoneIndex = 0; PerBoneIndex < ChildSampleData.PerBoneBlendData.Num(); ++PerBoneIndex)
+				{
+					ChildSampleData.PerBoneBlendData[PerBoneIndex] = ChildWeight;
+				}
+			}
+			else
+			{
+				// Guaranteed no child is using blend profile
 			}
 		}
 
-		if (!InstanceData->PerBoneSampleData.IsEmpty())
-		{
-			FBlendSampleData::NormalizeDataWeight(InstanceData->PerBoneSampleData);
-		}
+		// We need to normalize the per-bone weights of the active children
+		NormalizeChildPerBoneWeights(InstanceData);
 	}
 
 	void FBlendSmootherPerBoneCoreTrait::OnBlendTransition(FExecutionContext& Context, const TTraitBinding<IDiscreteBlend>& Binding, int32 OldChildIndex, int32 NewChildIndex) const
 	{
+		// First delegate to the traits below us
+		// Again, ISmoothBlend does most of the work for us
 		IDiscreteBlend::OnBlendTransition(Context, Binding, OldChildIndex, NewChildIndex);
 
 		const FSharedData* SharedData = Binding.GetSharedData<FSharedData>();
@@ -198,139 +233,214 @@ namespace UE::UAF
 			// We have a new child
 			check(NewChildIndex == NumChildren);
 
-			TSharedPtr<const IBlendProfileInterface> BlendProfileInterface = SmoothBlendPerBoneTrait.GetBlendProfile(Context, NewChildIndex);
+			FBlendData& ChildBlendData = InstanceData->PerChildBlendData.AddDefaulted_GetRef();
+			FBlendSampleData& NewSampleData = InstanceData->PerBoneSampleData.AddDefaulted_GetRef();
+			if (InstanceData->BlendProfileSkeleton)
+			{
+				NewSampleData.PerBoneBlendData.SetNumUninitialized(InstanceData->BlendProfileSkeleton->GetReferenceSkeleton().GetNum());
+			}
+		}
 
-			const FAlphaBlend* BlendState = DiscreteBlendTrait.GetBlendState(Context, NewChildIndex);
+		InstanceData->ChildOrdering.RemoveSwap(NewChildIndex);
+		InstanceData->ChildOrdering.Add(NewChildIndex);
 
-			FBlendData& ChildBlendData = InstanceData->PerChildBlendData.AddDefaulted_GetRef();
-			ChildBlendData.BlendProfileInterface = BlendProfileInterface;
-			ChildBlendData.StartAlpha = BlendState->GetAlpha();
+		{
+			// Setup the new child to blend in
+			FBlendData& ChildBlendData = InstanceData->PerChildBlendData[NewChildIndex];
+			FBlendSampleData& ChildSampleData = InstanceData->PerBoneSampleData[NewChildIndex];
+			
+			const UUAFBlendProfile* BlendProfile = SmoothBlendPerBoneTrait.GetBlendProfile(Context, NewChildIndex);
+			if (BlendProfile)
+			{
+				if (!InstanceData->BlendProfileSkeleton)
+				{
+					SetBlendProfileSkeleton(BlendProfile->GetSkeleton(), SharedData, InstanceData);
+				}
+				else if (BlendProfile->GetSkeleton() == InstanceData->BlendProfileSkeleton)
+				{
+					ChildBlendData.BlendProfile = BlendProfile;
+				}
+			}
 
-			FBlendSampleData& ChildSampleData = InstanceData->PerBoneSampleData.AddDefaulted_GetRef();
-			ChildSampleData.SampleDataIndex = NewChildIndex;
-			ChildSampleData.PerBoneBlendData.AddZeroed(BlendProfileInterface ? BlendProfileInterface->GetNumBlendEntries() : 0);
+			if (InstanceData->BlendProfileSkeleton)
+			{
+				const int32 BoneCount = InstanceData->BlendProfileSkeleton->GetReferenceSkeleton().GetNum();
+				for (int32 Index = 0; Index < BoneCount; ++Index)
+				{
+					ChildSampleData.PerBoneBlendData[Index] = 1.0f;
+				}
+			}
+
+			ChildBlendData.OldWeight = DiscreteBlendTrait.GetBlendWeight(Context, NewChildIndex);
+			ChildBlendData.bIsBlending = true;
 		}
+	}
 
-		for (int32 ChildIndex = 0; ChildIndex < NumChildren; ++ChildIndex)
-		{
-			FBlendData& ChildBlendData = InstanceData->PerChildBlendData[ChildIndex];
+	void FBlendSmootherPerBoneCoreTrait::OnBlendTerminated(FExecutionContext& Context, const TTraitBinding<IDiscreteBlend>& Binding, int32 ChildIndex) const
+	{
+		IDiscreteBlend::OnBlendTerminated(Context, Binding, ChildIndex);
+
+		FInstanceData* InstanceData = Binding.GetInstanceData<FInstanceData>();
 
-			const FAlphaBlend* BlendState = DiscreteBlendTrait.GetBlendState(Context, ChildIndex);
-			ChildBlendData.StartAlpha = BlendState->GetAlpha();
+		FBlendData& ChildBlendData = InstanceData->PerChildBlendData[ChildIndex];
+		ChildBlendData.bIsBlending = false;
+		ChildBlendData.BlendProfile = nullptr;
+		
+		FBlendSampleData& SampleData = InstanceData->PerBoneSampleData[ChildIndex];
+		for (int32 PerBoneIndex = 0; PerBoneIndex < SampleData.PerBoneBlendData.Num(); ++PerBoneIndex)
+		{
+			SampleData.PerBoneBlendData[PerBoneIndex] = 0.0f;
 		}
 	}
 
 	void FBlendSmootherPerBoneCoreTrait::InitializeInstanceData(FExecutionContext& Context, const FTraitBinding& Binding, const FSharedData* SharedData, FInstanceData* InstanceData)
 	{
+		check(InstanceData->PerBoneSampleData.IsEmpty());
 		check(InstanceData->PerChildBlendData.IsEmpty());
+		
+		TTraitBinding<IDiscreteBlend> DiscreteBlendTrait;
+		Binding.GetStackInterface(DiscreteBlendTrait);
 
 		TTraitBinding<ISmoothBlendPerBone> SmoothBlendPerBoneTrait;
 		Binding.GetStackInterface(SmoothBlendPerBoneTrait);
 
 		const uint32 NumChildren = IHierarchy::GetNumStackChildren(Context, Binding);
 
-		InstanceData->PerChildBlendData.SetNum(NumChildren);
 		InstanceData->PerBoneSampleData.SetNum(NumChildren);
+		InstanceData->PerChildBlendData.SetNum(NumChildren);
 
 		for (uint32 ChildIndex = 0; ChildIndex < NumChildren; ++ChildIndex)
 		{
-			TSharedPtr<const IBlendProfileInterface> BlendProfileInterface = SmoothBlendPerBoneTrait.GetBlendProfile(Context, ChildIndex);
+			InstanceData->ChildOrdering.Add(ChildIndex);
+			
+			const UUAFBlendProfile* BlendProfile = SmoothBlendPerBoneTrait.GetBlendProfile(Context, ChildIndex);
+
+			if (BlendProfile)
+			{
+				if (InstanceData->BlendProfileSkeleton)
+				{
+					if (BlendProfile->GetSkeleton() != InstanceData->BlendProfileSkeleton)
+					{
+						UE_LOGF(LogAnimation, Warning, "BlendSmootherPerBoneCoreTrait: Contains blend profiles using differing skeletons");
+					}
+				}
+				else
+				{
+					SetBlendProfileSkeleton(BlendProfile->GetSkeleton(), SharedData, InstanceData);
+				}
+			}
+		}
 
+		for (uint32 ChildIndex = 0; ChildIndex < NumChildren; ++ChildIndex)
+		{
 			FBlendData& ChildBlendData = InstanceData->PerChildBlendData[ChildIndex];
-			ChildBlendData.BlendProfileInterface = BlendProfileInterface;
-			ChildBlendData.StartAlpha = 0.0f;
+			ChildBlendData.OldWeight = DiscreteBlendTrait.GetBlendWeight(Context, ChildIndex);
+
+			const UUAFBlendProfile* BlendProfile = SmoothBlendPerBoneTrait.GetBlendProfile(Context, ChildIndex);
+			if (BlendProfile && BlendProfile->GetSkeleton() == InstanceData->BlendProfileSkeleton)
+			{
+				ChildBlendData.BlendProfile = BlendProfile;
+			}
 
-			if (BlendProfileInterface)
+			if (InstanceData->BlendProfileSkeleton)
 			{
-				const uint32 NumBlendEntries = BlendProfileInterface->GetNumBlendEntries();
 				FBlendSampleData& SampleData = InstanceData->PerBoneSampleData[ChildIndex];
-				SampleData.SampleDataIndex = ChildIndex;
-				SampleData.PerBoneBlendData.AddZeroed(NumBlendEntries);
+				SampleData.PerBoneBlendData.Init(1.0f, InstanceData->BlendProfileSkeleton->GetReferenceSkeleton().GetNum());
 			}
 		}
 	}
 
-	/*
-		FBlendProfilePerChildProviderTrait
-	*/
-
-	void FBlendProfilePerChildProviderTrait::PreUpdate(FUpdateTraversalContext& Context, const TTraitBinding<IUpdate>& Binding, const FTraitUpdateState& TraitState) const
+	void FBlendSmootherPerBoneCoreTrait::SetBlendProfileSkeleton(const USkeleton* Skeleton, const FSharedData* SharedData, FInstanceData* InstanceData)
 	{
-		const FSharedData* SharedData = Binding.GetSharedData<FSharedData>();
-		FInstanceData* InstanceData = Binding.GetInstanceData<FInstanceData>();
-
-		if (InstanceData->BlendProfileInterfaces.IsEmpty())
+		if (!ensure(Skeleton))
 		{
-			InitializeInstanceData(Context, Binding, SharedData, InstanceData);
+			return;
 		}
 
-		IUpdate::PreUpdate(Context, Binding, TraitState);
+		InstanceData->BlendProfileSkeleton = Skeleton;
+
+		const int32 NumChildren = InstanceData->PerBoneSampleData.Num();
+		const int32 NumBones = Skeleton->GetReferenceSkeleton().GetNum();
+
+		InstanceData->PerBoneTotals.SetNumZeroed(NumBones);
+
+		for (int32 ChildIndex = 0; ChildIndex < NumChildren; ++ChildIndex)
+		{
+			InstanceData->PerBoneSampleData[ChildIndex].PerBoneBlendData.Init(1.0f, NumBones);
+		}
 	}
 
-	TSharedPtr<const IBlendProfileInterface> FBlendProfilePerChildProviderTrait::GetBlendProfile(FExecutionContext& Context, const TTraitBinding<ISmoothBlendPerBone>& Binding, int32 ChildIndex) const
+	void FBlendSmootherPerBoneCoreTrait::NormalizeChildPerBoneWeights(TNonNullPtr<FInstanceData> InstanceData)
 	{
-		FInstanceData* InstanceData = Binding.GetInstanceData<FInstanceData>();
-
-		// Try find another trait that'll provide a custom blend profile before falling back to ours
-		if (TSharedPtr<const IBlendProfileInterface> BlendProfileInterface = ISmoothBlendPerBone::GetBlendProfile(Context, Binding, ChildIndex))
+		if (!InstanceData->BlendProfileSkeleton)
 		{
-			return BlendProfileInterface;
+			return; // No per-bone weights to normalize
 		}
+		
+		const int32 NumChildren = InstanceData->PerChildBlendData.Num();
+		check(InstanceData->ChildOrdering.Num() == NumChildren);
+		check(InstanceData->PerBoneSampleData.Num() == NumChildren);
 
-		if (InstanceData->BlendProfileInterfaces.IsValidIndex(ChildIndex))
+		if (NumChildren < 2)
 		{
-			return InstanceData->BlendProfileInterfaces[ChildIndex];
+			return;
 		}
+		
+		const int32 NumBones = InstanceData->PerBoneTotals.Num();
 
-		return nullptr;
-	}
+		for (int32 PerBoneIndex = 0; PerBoneIndex < NumBones; ++PerBoneIndex)
+		{
+			InstanceData->PerBoneTotals[PerBoneIndex] = 0.0f;
+		}
+		
+		for (int32 ChildIndex = 0; ChildIndex < NumChildren; ++ChildIndex)
+		{
+			const FBlendData& ChildBlendData = InstanceData->PerChildBlendData[ChildIndex];
+			if (!ChildBlendData.bIsBlending)
+			{
+				continue;
+			}
+			
+			const FBlendSampleData& ChildSampleData = InstanceData->PerBoneSampleData[ChildIndex];
 
-	void FBlendProfilePerChildProviderTrait::InitializeInstanceData(FExecutionContext& Context, const FTraitBinding& Binding, const FSharedData* SharedData, FInstanceData* InstanceData)
-	{
-		check(InstanceData->BlendProfileInterfaces.IsEmpty());
+			for (int32 PerBoneIndex = 0; PerBoneIndex < NumBones; ++PerBoneIndex)
+			{
+				InstanceData->PerBoneTotals[PerBoneIndex] += ChildSampleData.PerBoneBlendData[PerBoneIndex];
+			}
+		}
 
-		for (TObjectPtr<UHierarchyTable> BlendProfile : SharedData->BlendProfiles)
+		for (int32 ChildIndex = 0; ChildIndex < NumChildren; ++ChildIndex)
 		{
-			if (BlendProfile)
+			const FBlendData& ChildBlendData = InstanceData->PerChildBlendData[ChildIndex];
+			if (!ChildBlendData.bIsBlending)
+			{
+				continue;
+			}
+			
+			FBlendSampleData& ChildSampleData = InstanceData->PerBoneSampleData[ChildIndex];
+
+			for (int32 PerBoneIndex = 0; PerBoneIndex < NumBones; ++PerBoneIndex)
 			{
-				InstanceData->BlendProfileInterfaces.Add(MakeShared<FHierarchyTableBlendProfile>(BlendProfile, EBlendProfileMode::TimeFactor));
+				ChildSampleData.PerBoneBlendData[PerBoneIndex] /= InstanceData->PerBoneTotals[PerBoneIndex];
 			}
 		}
 	}
 
-	/*
-		FBlendProfileProviderTrait
-	*/
-
-	void FBlendProfileProviderTrait::PreUpdate(FUpdateTraversalContext& Context, const TTraitBinding<IUpdate>& Binding, const FTraitUpdateState& TraitState) const
+	const UUAFBlendProfile* FBlendSmootherPerBoneTrait::GetBlendProfile(FExecutionContext& Context, const TTraitBinding<ISmoothBlendPerBone>& Binding, int32 ChildIndex) const
 	{
 		const FSharedData* SharedData = Binding.GetSharedData<FSharedData>();
-		FInstanceData* InstanceData = Binding.GetInstanceData<FInstanceData>();
 
-		if (SharedData->TimeFactorBlendProfile && !InstanceData->BlendProfileInterface)
+		if (SharedData->BlendProfiles.IsValidIndex(ChildIndex))
 		{
-			InitializeInstanceData(Context, Binding, SharedData, InstanceData);
+			return SharedData->BlendProfiles[ChildIndex];
 		}
-
-		IUpdate::PreUpdate(Context, Binding, TraitState);
-	}
-
-	TSharedPtr<const IBlendProfileInterface> FBlendProfileProviderTrait::GetBlendProfile(FExecutionContext& Context, const TTraitBinding<ISmoothBlendPerBone>& Binding, int32 ChildIndex) const
-	{
-		const FSharedData* SharedData = Binding.GetSharedData<FSharedData>();
-		FInstanceData* InstanceData = Binding.GetInstanceData<FInstanceData>();
-
-		// Try find another trait that'll provide a custom blend profile before falling back to ours
-		if (TSharedPtr<const IBlendProfileInterface> BlendProfileInterface = ISmoothBlendPerBone::GetBlendProfile(Context, Binding, ChildIndex))
+		else if (!SharedData->BlendProfiles.IsEmpty())
 		{
-			return BlendProfileInterface;
+			return SharedData->BlendProfiles.Last();
+		}
+		else
+		{
+			return ISmoothBlendPerBone::GetBlendProfile(Context, Binding, ChildIndex);
 		}
-		
-		return InstanceData->BlendProfileInterface;
-	}
-
-	void FBlendProfileProviderTrait::InitializeInstanceData(FExecutionContext& Context, const FTraitBinding& Binding, const FSharedData* SharedData, FInstanceData* InstanceData)
-	{
-		InstanceData->BlendProfileInterface = MakeShared<FHierarchyTableBlendProfile>(SharedData->TimeFactorBlendProfile, EBlendProfileMode::TimeFactor);
 	}
 }
diff --git a/Engine/Plugins/Experimental/UAF/UAFAnimGraph/Source/UAFAnimGraph/Private/Traits/BlendSmootherPerBone.h b/Engine/Plugins/Experimental/UAF/UAFAnimGraph/Source/UAFAnimGraph/Private/Traits/BlendSmootherPerBone.h
--- a/Engine/Plugins/Experimental/UAF/UAFAnimGraph/Source/UAFAnimGraph/Private/Traits/BlendSmootherPerBone.h
+++ b/Engine/Plugins/Experimental/UAF/UAFAnimGraph/Source/UAFAnimGraph/Private/Traits/BlendSmootherPerBone.h
@@ -5,47 +5,26 @@
 #include "CoreMinimal.h"
 
 #include "Animation/AnimationAsset.h"
-#include "Curves/CurveFloat.h"
+#include "Traits/BlendSmootherPerBoneTraitData.h"
 #include "TraitCore/Trait.h"
 #include "TraitInterfaces/IDiscreteBlend.h"
 #include "TraitInterfaces/IEvaluate.h"
 #include "TraitInterfaces/ISmoothBlend.h"
 #include "TraitInterfaces/IUpdate.h"
 #include "TraitInterfaces/ISmoothBlendPerBone.h"
-#include "HierarchyTable.h"
 #include "HierarchyTableBlendProfile.h"
 
 #include "BlendSmootherPerBone.generated.h"
 
-/** An additive trait that smoothly blends with per-bone weights. */
-USTRUCT(meta = (DisplayName = "Blend Smoother Per Bone Core", ShowTooltip=true))
-struct FAnimNextBlendSmootherPerBoneCoreTraitSharedData : public FAnimNextTraitSharedData
-{
-	GENERATED_BODY()
-};
+class UUAFBlendProfile;
 
-USTRUCT(meta = (DisplayName = "Blend Profile Per Child Provider"))
-struct FAnimNextBlendProfilePerChildProviderTraitSharedData : public FAnimNextBlendSmootherPerBoneCoreTraitSharedData
+USTRUCT(meta = (DisplayName = "Blend Smoother Per Bone"))
+struct FUAFBlendSmootherPerBoneTraitSharedData : public FAnimNextBlendSmootherPerBoneCoreTraitSharedData
 {
 	GENERATED_BODY()
 
-	/** Blend profile that configures how fast to blend each bone. */
 	UPROPERTY(EditAnywhere, Category = "Default", meta = (Inline))
-	TArray<TObjectPtr<UHierarchyTable>> BlendProfiles;
-};
-
-/**
- * Serves firstly as a passthrough node that propagates blend profiles of child traits.
- * If no child trait provides a blend profile then uses the held blend profile.
- */
-USTRUCT(meta = (DisplayName = "Blend Profile Provider", ShowTooltip=true))
-struct FAnimNextBlendProfileProviderTraitSharedData : public FAnimNextTraitSharedData
-{
-	GENERATED_BODY()
-
-	/** Blend profile that configures how fast to blend each bone. */
-	UPROPERTY(EditAnywhere, Category = "Default", meta = (Inline))
-	TObjectPtr<UHierarchyTable> TimeFactorBlendProfile;
+	TArray<TObjectPtr<UUAFBlendProfile>> BlendProfiles;
 };
 
 namespace UE::UAF
@@ -64,18 +43,32 @@ namespace UE::UAF
 		// Struct for tracking blends for each pose
 		struct FBlendData
 		{
-			TSharedPtr<const IBlendProfileInterface> BlendProfileInterface;
+			// The blend profile to use to blend this child in with
+			// All blend profiles must use the same skeleton
+			const UUAFBlendProfile* BlendProfile = nullptr;
+
+			// The cached value of this child's weight last update
+			float OldWeight = 0.0f;
 
-			float StartAlpha = 0.0f;
+			// Flag to skip a child if its fully blended out
+			bool bIsBlending = false;
 		};
 
 		struct FInstanceData : FTrait::FInstanceData
 		{
-			// Blend state per child
 			TArray<FBlendData> PerChildBlendData;
 
 			// Per-bone blending data for each child
 			TArray<FBlendSampleData> PerBoneSampleData;
+
+			// Scratch data for norma
+			TArray<float> PerBoneTotals;
+
+			// The ordering of children
+			TArray<int32> ChildOrdering;
+
+			// The common skeleton all blend profiles must share
+			const USkeleton* BlendProfileSkeleton;
 		};
 
 		// IEvaluate impl
@@ -86,6 +79,7 @@ namespace UE::UAF
 
 		// IDiscreteBlend impl
 		virtual void OnBlendTransition(FExecutionContext& Context, const TTraitBinding<IDiscreteBlend>& Binding, int32 OldChildIndex, int32 NewChildIndex) const override;
+		virtual void OnBlendTerminated(FExecutionContext& Context, const TTraitBinding<IDiscreteBlend>& Binding, int32 ChildIndex) const;
 
 #if WITH_EDITOR
 		virtual bool IsHidden() const override { return false; }
@@ -94,68 +88,17 @@ namespace UE::UAF
 	private:
 		// Internal impl
 		static void InitializeInstanceData(FExecutionContext& Context, const FTraitBinding& Binding, const FSharedData* SharedData, FInstanceData* InstanceData);
+		static void SetBlendProfileSkeleton(const USkeleton* Skeleton, const FSharedData* SharedData, FInstanceData* InstanceData);
+		static void NormalizeChildPerBoneWeights(TNonNullPtr<FInstanceData> InstanceData);
 	};
 
-	/**
-	 * FBlendProfilePerChildProviderTrait
-	 */
-	struct FBlendProfilePerChildProviderTrait : FAdditiveTrait, IUpdate, ISmoothBlendPerBone
-	{
-		DECLARE_ANIM_TRAIT(FBlendProfilePerChildProviderTrait, FAdditiveTrait)
-
-		using FSharedData = FAnimNextBlendProfilePerChildProviderTraitSharedData;
-
-		struct FInstanceData : FTrait::FInstanceData
-		{
-			// Blend state per child
-			TArray<TSharedPtr<const FHierarchyTableBlendProfile>> BlendProfileInterfaces;
-		};
-
-		// IUpdate impl
-		virtual void PreUpdate(FUpdateTraversalContext& Context, const TTraitBinding<IUpdate>& Binding, const FTraitUpdateState& TraitState) const override;
-
-		// ISmoothBlendPerBone impl
-		virtual TSharedPtr<const IBlendProfileInterface> GetBlendProfile(FExecutionContext& Context, const TTraitBinding<ISmoothBlendPerBone>& Binding, int32 ChildIndex) const override;
-
-#if WITH_EDITOR
-		virtual bool IsHidden() const override { return false; } // override the base, as it is hidden
-#endif
-
-	private:
-		// Internal impl
-		static void InitializeInstanceData(FExecutionContext& Context, const FTraitBinding& Binding, const FSharedData* SharedData, FInstanceData* InstanceData);
-	};
-
-	/**
-	 * FBlendProfileProviderTrait
-	 * 
-	 * Serves firstly as a passthrough node that propagates blend profiles of child traits.
-	 * If no child trait provides a blend profile then uses the held blend profile.
-	 */
-	struct FBlendProfileProviderTrait : FAdditiveTrait, IUpdate, ISmoothBlendPerBone
+	struct FBlendSmootherPerBoneTrait : public FBlendSmootherPerBoneCoreTrait, ISmoothBlendPerBone
 	{
-		DECLARE_ANIM_TRAIT(FBlendProfileProviderTrait, FAdditiveTrait)
+		DECLARE_ANIM_TRAIT(FBlendSmootherPerBoneTrait, FBlendSmootherPerBoneCoreTrait)
 
-		using FSharedData = FAnimNextBlendProfileProviderTraitSharedData;
+		using FSharedData = FUAFBlendSmootherPerBoneTraitSharedData;
 
-		struct FInstanceData : FTrait::FInstanceData
-		{
-			// Blend state per child
-			TSharedPtr<const FHierarchyTableBlendProfile> BlendProfileInterface;
-		};
-
-		// IUpdate impl
-		virtual void PreUpdate(FUpdateTraversalContext& Context, const TTraitBinding<IUpdate>& Binding, const FTraitUpdateState& TraitState) const override;
-
-		// ISmoothBlendPerBone impl
-		virtual TSharedPtr<const IBlendProfileInterface> GetBlendProfile(FExecutionContext& Context, const TTraitBinding<ISmoothBlendPerBone>& Binding, int32 ChildIndex) const override;
-
-#if WITH_EDITOR
-		virtual bool IsHidden() const override { return false; } // override the base, as it is hidden
-#endif
-
-	private:
-		// Internal impl
-		static void InitializeInstanceData(FExecutionContext& Context, const FTraitBinding& Binding, const FSharedData* SharedData, FInstanceData* InstanceData);
+		// ISmoothBlendPerBone
+		virtual const UUAFBlendProfile* GetBlendProfile(FExecutionContext& Context, const TTraitBinding<ISmoothBlendPerBone>& Binding, int32 ChildIndex) const override;
 	};
 }
diff --git a/Engine/Plugins/Experimental/UAF/UAFAnimGraph/Source/UAFAnimGraph/Private/Traits/BlendStackTrait.cpp b/Engine/Plugins/Experimental/UAF/UAFAnimGraph/Source/UAFAnimGraph/Private/Traits/BlendStackTrait.cpp
--- a/Engine/Plugins/Experimental/UAF/UAFAnimGraph/Source/UAFAnimGraph/Private/Traits/BlendStackTrait.cpp
+++ b/Engine/Plugins/Experimental/UAF/UAFAnimGraph/Source/UAFAnimGraph/Private/Traits/BlendStackTrait.cpp
@@ -479,7 +479,7 @@ namespace UE::UAF
 		return InstanceData->ChildGraphs[InChildIndex].Instance.Get();
 	}
 
-	TSharedPtr<const IBlendProfileInterface> FBlendStackCoreTrait::GetBlendProfile(FExecutionContext& Context, const TTraitBinding<ISmoothBlendPerBone>& Binding, int32 ChildIndex) const
+	const UUAFBlendProfile* FBlendStackCoreTrait::GetBlendProfile(FExecutionContext& Context, const TTraitBinding<ISmoothBlendPerBone>& Binding, int32 ChildIndex) const
 	{
 		const FInstanceData* InstanceData = Binding.GetInstanceData<FInstanceData>();
 		if (InstanceData->ChildGraphs.IsValidIndex(ChildIndex))
diff --git a/Engine/Plugins/Experimental/UAF/UAFAnimGraph/Source/UAFAnimGraph/Public/EvaluationVM/Tasks/ApplyAdditiveKeyframePerBone.h b/Engine/Plugins/Experimental/UAF/UAFAnimGraph/Source/UAFAnimGraph/Public/EvaluationVM/Tasks/ApplyAdditiveKeyframePerBone.h
--- a/Engine/Plugins/Experimental/UAF/UAFAnimGraph/Source/UAFAnimGraph/Public/EvaluationVM/Tasks/ApplyAdditiveKeyframePerBone.h
+++ b/Engine/Plugins/Experimental/UAF/UAFAnimGraph/Source/UAFAnimGraph/Public/EvaluationVM/Tasks/ApplyAdditiveKeyframePerBone.h
@@ -11,6 +11,7 @@
 
 struct FBlendSampleData;
 class IBlendProfileInterface;
+class UUAFBlendMask;
 
 /*
  * Apply Additive Keyframe Per Bone Task
@@ -41,25 +42,26 @@ struct FUAFApplyAdditiveKeyframePerBoneTask : public FAnimNextApplyAdditiveKeyfr
  * Apply Additive Keyframe Per Bone Task
  *
  * This pops the top two keyframes from the VM keyframe stack, 
- * it applies an additive keyframe with a weight per bone based on the given blend profile onto its base and pushes back the result onto the stack.
- * If no valid blend profile is provided it will revert back to a full pose additive application. 
+ * it applies an additive keyframe with a weight per bone based on the given blend mask onto its base and pushes back the result onto the stack.
+ * If no valid blend mask is provided it will revert back to a full pose additive application. 
  * 
  * The top pose should be the additive keyframe and the second to the top the base keyframe.
  * 
- * This task currently only supports hierarchy table based blend profiles.
+ * This task currently only supports skeleton-based blend masks.
  */
 USTRUCT()
-struct FUAFApplyAdditiveKeyframeWithBlendProfileTask : public FAnimNextApplyAdditiveKeyframeTask
+struct FUAFApplyAdditiveKeyframeWithBlendMaskTask : public FAnimNextApplyAdditiveKeyframeTask
 {
 	GENERATED_BODY()
 
-	DECLARE_ANIM_EVALUATION_TASK(FUAFApplyAdditiveKeyframeWithBlendProfileTask)
+	DECLARE_ANIM_EVALUATION_TASK(FUAFApplyAdditiveKeyframeWithBlendMaskTask)
 	
-	static UE_API FUAFApplyAdditiveKeyframeWithBlendProfileTask Make(const float InBlendWeight, const TSharedPtr<const IBlendProfileInterface>& InBlendProfileInterface);
+	static UE_API FUAFApplyAdditiveKeyframeWithBlendMaskTask Make(const float InBlendWeight, const UUAFBlendMask* InBlendMask);
 	
 	UE_API virtual void Execute(UE::UAF::FEvaluationVM& VM) const override;
 	
-	TSharedPtr<const IBlendProfileInterface> BlendProfileInterface = nullptr;
+	UPROPERTY()
+	TObjectPtr<const UUAFBlendMask> BlendMask = nullptr;
 };
 
 
diff --git a/Engine/Plugins/Experimental/UAF/UAFAnimGraph/Source/UAFAnimGraph/Public/EvaluationVM/Tasks/BlendKeyframesPerBone.h b/Engine/Plugins/Experimental/UAF/UAFAnimGraph/Source/UAFAnimGraph/Public/EvaluationVM/Tasks/BlendKeyframesPerBone.h
--- a/Engine/Plugins/Experimental/UAF/UAFAnimGraph/Source/UAFAnimGraph/Public/EvaluationVM/Tasks/BlendKeyframesPerBone.h
+++ b/Engine/Plugins/Experimental/UAF/UAFAnimGraph/Source/UAFAnimGraph/Public/EvaluationVM/Tasks/BlendKeyframesPerBone.h
@@ -8,6 +8,7 @@
 #include "EvaluationVM/EvaluationTask.h"
 #include "EvaluationVM/Tasks/BlendKeyframes.h"
 #include "Animation/AttributesContainer.h"
+#include "UAF/BlendMask/UAFBlendMask.h"
 
 #include "BlendKeyframesPerBone.generated.h"
 
@@ -37,8 +38,7 @@ struct FAnimNextBlendOverwriteKeyframePerBoneWithScaleTask : public FAnimNextBle
 
 	DECLARE_ANIM_EVALUATION_TASK(FAnimNextBlendOverwriteKeyframePerBoneWithScaleTask)
 
-	static UE_API FAnimNextBlendOverwriteKeyframePerBoneWithScaleTask Make(const IBlendProfileInterface* BlendProfile, const FBlendSampleData& BlendData, float ScaleFactor);
-	static UE_API FAnimNextBlendOverwriteKeyframePerBoneWithScaleTask Make(const FBlendSampleData& BlendData, float ScaleFactor, bool bPerBoneDataUsesSkeletonIndex);
+	static UE_API FAnimNextBlendOverwriteKeyframePerBoneWithScaleTask Make(const USkeleton* Skeleton, const FBlendSampleData& BlendData, float ScaleFactor);
 
 	// Task entry point
 	UE_API virtual void Execute(UE::UAF::FEvaluationVM& VM) const override;
@@ -48,13 +48,8 @@ struct FAnimNextBlendOverwriteKeyframePerBoneWithScaleTask : public FAnimNextBle
 	//UPROPERTY()
 	const FBlendSampleData* BlendData = nullptr;
 
-	// The blend profile to use
-	// An optional blend profile provides a blend weight per bone
-	const IBlendProfileInterface* BlendProfile = nullptr;
-
-	// True if the provided bone weight data uses the full skeleton bone set of the target skeleton for indexing instead of the current LOD's
-	// This is used to build the weight lookup map if no blend profile is used 
-	bool bPerBoneDataUsesSkeletonIndex = false;
+	// Optional skeleton to perform skeleton remapping with
+	const USkeleton* SourceSkeleton = nullptr;
 };
 
 /*
@@ -78,8 +73,7 @@ struct FAnimNextBlendAddKeyframePerBoneWithScaleTask : public FAnimNextBlendAddK
 
 	DECLARE_ANIM_EVALUATION_TASK(FAnimNextBlendAddKeyframePerBoneWithScaleTask)
 
-	static UE_API FAnimNextBlendAddKeyframePerBoneWithScaleTask Make(const IBlendProfileInterface* BlendProfile, const FBlendSampleData& BlendDataA, const FBlendSampleData& BlendDataB, float ScaleFactor);
-	static UE_API FAnimNextBlendAddKeyframePerBoneWithScaleTask Make(const FBlendSampleData& BlendDataA, const FBlendSampleData& BlendDataB, float ScaleFactor, bool bPerBoneDataUsesSkeletonIndex);
+	static UE_API FAnimNextBlendAddKeyframePerBoneWithScaleTask Make(const USkeleton* Skeleton, const FBlendSampleData& BlendDataA, const FBlendSampleData& BlendDataB, float ScaleFactor);
 
 	// Task entry point
 	UE_API virtual void Execute(UE::UAF::FEvaluationVM& VM) const override;
@@ -94,13 +88,8 @@ struct FAnimNextBlendAddKeyframePerBoneWithScaleTask : public FAnimNextBlendAddK
 	//UPROPERTY()
 	const FBlendSampleData* BlendDataB = nullptr;
 
-	// The blend profile to use
-	// An optional blend profile provides a blend weight per bone
-	const IBlendProfileInterface* BlendProfile = nullptr;
-
-	// True if the provided bone weight data uses the full skeleton bone set of the target skeleton for indexing instead of the current LOD's
-	// This is used to build the weight lookup map if no blend profile is used 
-	bool bPerBoneDataUsesSkeletonIndex = false;
+	// Optional skeleton to perform skeleton remapping with
+	const USkeleton* SourceSkeleton = nullptr;
 };
 
 
@@ -111,32 +100,24 @@ struct FAnimNextBlendKeyframePerBoneWithScaleTask : public FAnimNextBlendAddKeyf
 
 	DECLARE_ANIM_EVALUATION_TASK(FAnimNextBlendKeyframePerBoneWithScaleTask)
 
-	struct FMaskedAttributeWeight
-	{
-		UE::Anim::FAttributeId Attribute;
-		float Weight;
-
-		FMaskedAttributeWeight(const UE::Anim::FAttributeId InAttribute, const float InWeight)
-			: Attribute(InAttribute)
-			, Weight(InWeight)
-		{
-		}
-	};
+	static UE_API FAnimNextBlendKeyframePerBoneWithScaleTask Make(const UUAFBlendMask* InBlendMask, const float ScaleFactor);
 
-	static UE_API FAnimNextBlendKeyframePerBoneWithScaleTask Make(TSharedPtr<IInterpolationIndexProvider> BlendProfile, const USkeleton* Skeleton, const TArray<float>& BoneMaskWeights, const UE::Anim::TNamedValueArray<FDefaultAllocator, UE::Anim::FCurveElement>& CurveMaskWeights, const TArray<FMaskedAttributeWeight>& AttributeMaskWeights, float ScaleFactor);
+	// Curve Weights and Attribute Weights can be provided as empty views
+	static UE_API FAnimNextBlendKeyframePerBoneWithScaleTask Make(
+		const USkeleton* Skeleton,
+		const UUAFBlendMask::FSkeletonBoneWeightArray* InBoneWeights,
+		const UUAFBlendMask::FSkeletonCurveWeightArray* InCurveWeights,
+		const UUAFBlendMask::FSkeletonAttributeWeightArray* InAttributeWeights,
+		float ScaleFactor);
 
 	// Task entry point
 	UE_API virtual void Execute(UE::UAF::FEvaluationVM& VM) const override;
 
 private:
-	TArray<float> BoneMaskWeights;
-
-	UE::Anim::TNamedValueArray<FDefaultAllocator, UE::Anim::FCurveElement> CurveMaskWeights;
-
-	TArray<FMaskedAttributeWeight> AttributeMaskWeights;
-
-	TSharedPtr<IInterpolationIndexProvider> BlendProfile = nullptr;
-
+	const UUAFBlendMask::FSkeletonBoneWeightArray* BoneMaskWeights;
+	const UUAFBlendMask::FSkeletonCurveWeightArray* CurveMaskWeights;
+	const UUAFBlendMask::FSkeletonAttributeWeightArray* AttributeMaskWeights;
+	
 	const USkeleton* SourceSkeleton = nullptr;
 };
 
diff --git a/Engine/Plugins/Experimental/UAF/UAFAnimGraph/Source/UAFAnimGraph/Public/TraitInterfaces/IBlendStack.h b/Engine/Plugins/Experimental/UAF/UAFAnimGraph/Source/UAFAnimGraph/Public/TraitInterfaces/IBlendStack.h
--- a/Engine/Plugins/Experimental/UAF/UAFAnimGraph/Source/UAFAnimGraph/Public/TraitInterfaces/IBlendStack.h
+++ b/Engine/Plugins/Experimental/UAF/UAFAnimGraph/Source/UAFAnimGraph/Public/TraitInterfaces/IBlendStack.h
@@ -14,6 +14,8 @@
 
 #define UE_API UAFANIMGRAPH_API
 
+class UUAFBlendProfile;
+
 namespace UE::UAF
 {
 	/**
@@ -54,7 +56,7 @@ namespace UE::UAF
 			FAlphaBlendArgs BlendArgs;
 
 			// Optional blend profile for this graph request
-			TSharedPtr<const IBlendProfileInterface> BlendProfile;
+			TObjectPtr<const UUAFBlendProfile> BlendProfile;
 
 			// The template graph to use for the new graph instance
 			TObjectPtr<const UUAFAnimGraph> AnimationGraph;
diff --git a/Engine/Plugins/Experimental/UAF/UAFAnimGraph/Source/UAFAnimGraph/Public/TraitInterfaces/ISmoothBlendPerBone.h b/Engine/Plugins/Experimental/UAF/UAFAnimGraph/Source/UAFAnimGraph/Public/TraitInterfaces/ISmoothBlendPerBone.h
--- a/Engine/Plugins/Experimental/UAF/UAFAnimGraph/Source/UAFAnimGraph/Public/TraitInterfaces/ISmoothBlendPerBone.h
+++ b/Engine/Plugins/Experimental/UAF/UAFAnimGraph/Source/UAFAnimGraph/Public/TraitInterfaces/ISmoothBlendPerBone.h
@@ -9,8 +9,8 @@
 
 #define UE_API UAFANIMGRAPH_API
 
-class IBlendProfileInterface;
 class UCurveFloat;
+class UUAFBlendProfile;
 
 namespace UE::UAF
 {
@@ -24,7 +24,7 @@ namespace UE::UAF
 		DECLARE_ANIM_TRAIT_INTERFACE(ISmoothBlendPerBone)
 
 		// Returns the desired blend profile for the specified child
-		UE_API virtual TSharedPtr<const IBlendProfileInterface> GetBlendProfile(FExecutionContext& Context, const TTraitBinding<ISmoothBlendPerBone>& Binding, int32 ChildIndex) const;
+		UE_API virtual const UUAFBlendProfile* GetBlendProfile(FExecutionContext& Context, const TTraitBinding<ISmoothBlendPerBone>& Binding, int32 ChildIndex) const;
 
 #if WITH_EDITOR
 		UE_API virtual const FText& GetDisplayName() const override;
@@ -39,7 +39,7 @@ namespace UE::UAF
 	struct TTraitBinding<ISmoothBlendPerBone> : FTraitBinding
 	{
 		// @see ISmoothBlendPerBone::GetBlendProfile
-		const TSharedPtr<const IBlendProfileInterface> GetBlendProfile(FExecutionContext& Context, int32 ChildIndex) const
+		const UUAFBlendProfile* GetBlendProfile(FExecutionContext& Context, int32 ChildIndex) const
 		{
 			return GetInterface()->GetBlendProfile(Context, *this, ChildIndex);
 		}
diff --git a/Engine/Plugins/Experimental/UAF/UAFAnimGraph/Source/UAFAnimGraph/Public/Traits/ApplyAdditivePerBoneTraitData.h b/Engine/Plugins/Experimental/UAF/UAFAnimGraph/Source/UAFAnimGraph/Public/Traits/ApplyAdditivePerBoneTraitData.h
--- a/Engine/Plugins/Experimental/UAF/UAFAnimGraph/Source/UAFAnimGraph/Public/Traits/ApplyAdditivePerBoneTraitData.h
+++ b/Engine/Plugins/Experimental/UAF/UAFAnimGraph/Source/UAFAnimGraph/Public/Traits/ApplyAdditivePerBoneTraitData.h
@@ -7,7 +7,7 @@
 
 #include "ApplyAdditivePerBoneTraitData.generated.h"
 
-class UHierarchyTable;
+class UUAFBlendMask;
 
 namespace UE::UAF
 {
@@ -26,10 +26,9 @@ struct FApplyAdditivePerBoneTraitSharedData : public FAnimNextTraitSharedData
 	UPROPERTY()
 	FAnimNextTraitHandle Additive;
 	
-	// TODO: To be changed over to UStandaloneBlendProfile
-	// HierarchyTable Mask asset
+	// Blend mask asset
 	UPROPERTY(EditAnywhere, Category = "Default")
-	TObjectPtr<UHierarchyTable> BlendProfile = nullptr;
+	TObjectPtr<UUAFBlendMask> BlendMask = nullptr;
 	
 	// How much to apply of the additive 
 	UPROPERTY(EditAnywhere, Category = "Default")
@@ -42,7 +41,7 @@ struct FApplyAdditivePerBoneTraitSharedData : public FAnimNextTraitSharedData
 	// Latent pin support boilerplate
 	#define TRAIT_LATENT_PROPERTIES_ENUMERATOR(GeneratorMacro) \
 	GeneratorMacro(Alpha) \
-	GeneratorMacro(BlendProfile) \
+	GeneratorMacro(BlendMask) \
 	GeneratorMacro(bMeshSpaceBlend) \
 
 	GENERATE_TRAIT_LATENT_PROPERTIES(FApplyAdditivePerBoneTraitSharedData, TRAIT_LATENT_PROPERTIES_ENUMERATOR)
diff --git a/Engine/Plugins/Experimental/UAF/UAFAnimGraph/Source/UAFAnimGraph/Public/Traits/BlendLayerTraitData.h b/Engine/Plugins/Experimental/UAF/UAFAnimGraph/Source/UAFAnimGraph/Public/Traits/BlendLayerTraitData.h
--- a/Engine/Plugins/Experimental/UAF/UAFAnimGraph/Source/UAFAnimGraph/Public/Traits/BlendLayerTraitData.h
+++ b/Engine/Plugins/Experimental/UAF/UAFAnimGraph/Source/UAFAnimGraph/Public/Traits/BlendLayerTraitData.h
@@ -8,7 +8,7 @@
 
 #include "BlendLayerTraitData.generated.h"
 
-class UHierarchyTable;
+class UUAFBlendMask;
 
 /** A trait that can blend a layer into an input. */
 USTRUCT(meta = (DisplayName = "Blend Layer", ShowTooltip=true))
@@ -26,11 +26,10 @@ struct FUAFBlendLayerTraitSharedData : public FAnimNextTraitSharedData
 	UPROPERTY(EditAnywhere, Category = "Default")
 	float BlendWeight = 1.0f;
 
-	/** Blend profile that configures how fast to blend each bone. */
-	// TODO: Can't show list of blend profiles, we need to find a skeleton to perform the lookup with
+	/** Blend mask that configures what bones are included in the blend. */
 	UPROPERTY(EditAnywhere, Category = "Default", meta = (Inline))
-	TObjectPtr<UHierarchyTable> BlendProfile = nullptr;
-
+	TObjectPtr<UUAFBlendMask> BlendMask = nullptr;
+	
 	// Whether or not to blend in mesh space or local space
 	UPROPERTY(EditAnywhere, Category = "Default")
 	bool bBlendInMeshSpace = false;
@@ -41,6 +40,7 @@ struct FUAFBlendLayerTraitSharedData : public FAnimNextTraitSharedData
 
 	#define TRAIT_LATENT_PROPERTIES_ENUMERATOR(GeneratorMacro) \
 	GeneratorMacro(BlendWeight) \
+	GeneratorMacro(BlendMask) \
 	GeneratorMacro(bBlendInMeshSpace) \
 	GeneratorMacro(bAlwaysUpdateChildBlend) \
 
diff --git a/Engine/Plugins/Experimental/UAF/UAFAnimGraph/Source/UAFAnimGraph/Public/Traits/BlendSmootherPerBoneTraitData.h b/Engine/Plugins/Experimental/UAF/UAFAnimGraph/Source/UAFAnimGraph/Public/Traits/BlendSmootherPerBoneTraitData.h
new file mode 100644
--- /dev/null
+++ b/Engine/Plugins/Experimental/UAF/UAFAnimGraph/Source/UAFAnimGraph/Public/Traits/BlendSmootherPerBoneTraitData.h
@@ -0,0 +1,15 @@
+// Copyright Epic Games, Inc. All Rights Reserved.
+
+#pragma once
+
+#include "CoreMinimal.h"
+#include "TraitCore/TraitSharedData.h"
+
+#include "BlendSmootherPerBoneTraitData.generated.h"
+
+/** An additive trait that smoothly blends with per-bone weights. */
+USTRUCT(meta = (DisplayName = "Blend Smoother Per Bone Core", ShowTooltip=true))
+struct FAnimNextBlendSmootherPerBoneCoreTraitSharedData : public FAnimNextTraitSharedData
+{
+	GENERATED_BODY()
+};
\ No newline at end of file
diff --git a/Engine/Plugins/Experimental/UAF/UAFLayering/Source/UAFLayeringUncookedOnly/Internal/Layers/UAFLayerDefaultBlendProvider.h b/Engine/Plugins/Experimental/UAF/UAFLayering/Source/UAFLayeringUncookedOnly/Internal/Layers/UAFLayerDefaultBlendProvider.h
--- a/Engine/Plugins/Experimental/UAF/UAFLayering/Source/UAFLayeringUncookedOnly/Internal/Layers/UAFLayerDefaultBlendProvider.h
+++ b/Engine/Plugins/Experimental/UAF/UAFLayering/Source/UAFLayeringUncookedOnly/Internal/Layers/UAFLayerDefaultBlendProvider.h
@@ -12,8 +12,9 @@
 #include "UAFLayerDefaultBlendProvider.generated.h"
 
 class UAnimNextController;
-class UHierarchyTable;
 class URigVMPin;
+class UUAFBlendMask;
+class UUAFBlendProfile;
 class UUAFLayer;
 class FUAFDefaultLayerBlendStructureDataProvider;
 
@@ -64,7 +65,7 @@ public:
 	// TODO: This should be a binding
 	// An optional blend mask to apply when blending this layers pose result with the previous layer result 
 	UPROPERTY(EditAnywhere, Category = "Layer")
-	TObjectPtr<UHierarchyTable> BlendMask = nullptr;
+	TObjectPtr<UUAFBlendMask> BlendMask = nullptr;
 	
 	// TODO: This should be a binding
 	// The custom blend curve to use when blending this layer in or out 
@@ -75,7 +76,7 @@ public:
 	// The blend profile to use to add more fine grain control over how certain bones should blend in or out 
 	// This can be used additionally to the blend mask that is used to apply this layer 
 	UPROPERTY(EditAnywhere, Category = "Layer|BlendSettings")
-	TObjectPtr<UHierarchyTable> TimeBasedBlendProfile = nullptr;
+	TObjectPtr<UUAFBlendProfile> BlendProfile = nullptr;
 	
 	// TODO: In the future we should be able to read this out of the blend mask 
 	// If a blend mask is applied, this is the named set for abstract hierarchies that mask relates to
diff --git a/Engine/Plugins/Experimental/UAF/UAFLayering/Source/UAFLayeringUncookedOnly/Private/Layers/UAFLayerDefaultBlendProvider.cpp b/Engine/Plugins/Experimental/UAF/UAFLayering/Source/UAFLayeringUncookedOnly/Private/Layers/UAFLayerDefaultBlendProvider.cpp
--- a/Engine/Plugins/Experimental/UAF/UAFLayering/Source/UAFLayeringUncookedOnly/Private/Layers/UAFLayerDefaultBlendProvider.cpp
+++ b/Engine/Plugins/Experimental/UAF/UAFLayering/Source/UAFLayeringUncookedOnly/Private/Layers/UAFLayerDefaultBlendProvider.cpp
@@ -5,10 +5,11 @@
 
 #include "AnimNextController.h"
 #include "AnimNextRigVMAssetEditorData.h"
-#include "HierarchyTable.h"
 #include "ISinglePropertyView.h"
 #include "UAFLayeringStyle.h"
 #include "UAFLayerStack.h"
+#include "UAF/BlendMask/UAFBlendMask.h"
+#include "UAF/BlendProfile/UAFBlendProfile.h"
 #include "Graph/RigUnit_AnimNextTraitStack.h"
 #include "Layers/UAFLayer.h"
 #include "RigVMCore/RigVMStruct.h"
@@ -114,14 +115,14 @@ URigVMPin* FUAFDefaultBlendProvider::CreateBlendGraphTrait(UE::UAF::Layering::FL
 	}
 	
 	FName BlendBaseTraitName = NAME_None;
-	const bool bUsesPerBoneData = BlendMask || TimeBasedBlendProfile;
+	const bool bUsesPerBoneData = BlendMask || BlendProfile;
 	if (BlendMode == EUAFLayerBlendMode::Additive)
 	{
 		if (bUsesPerBoneData)
 		{
 			UE::UAF::FApplyAdditivePerBoneTraitSharedData DefaultAdditiveTraitData;
 			DefaultAdditiveTraitData.Alpha = LayerWeight;
-			DefaultAdditiveTraitData.BlendProfile = BlendMask;
+			DefaultAdditiveTraitData.BlendMask = BlendMask;
 			
 			// Create trait 
 			BlendBaseTraitName = LayerCreationContext.GraphController->AddTraitStruct(BlendTraitStackNode, TInstancedStruct<UE::UAF::FApplyAdditivePerBoneTraitSharedData>::Make(DefaultAdditiveTraitData), INDEX_NONE, false, false);
@@ -150,7 +151,7 @@ URigVMPin* FUAFDefaultBlendProvider::CreateBlendGraphTrait(UE::UAF::Layering::FL
 			// Setup default trait values 
 			FUAFBlendLayerTraitSharedData DefaultTraitData; 
 			DefaultTraitData.BlendWeight = LayerWeight;
-			DefaultTraitData.BlendProfile = BlendMask;
+			DefaultTraitData.BlendMask = BlendMask;
 			DefaultTraitData.bAlwaysUpdateChildBlend = GetAlwaysUpdateChildren();
 			
 			// Create trait 
diff --git a/Engine/Plugins/Experimental/UAF/UAFPoseSearch/Source/UAFPoseSearch/Internal/MotionMatchingTraitData.h b/Engine/Plugins/Experimental/UAF/UAFPoseSearch/Source/UAFPoseSearch/Internal/MotionMatchingTraitData.h
--- a/Engine/Plugins/Experimental/UAF/UAFPoseSearch/Source/UAFPoseSearch/Internal/MotionMatchingTraitData.h
+++ b/Engine/Plugins/Experimental/UAF/UAFPoseSearch/Source/UAFPoseSearch/Internal/MotionMatchingTraitData.h
@@ -12,6 +12,8 @@
 #include "TraitInterfaces/IGroupSynchronization.h"
 #include "MotionMatchingTraitData.generated.h"
 
+class UUAFBlendProfile;
+
 USTRUCT(meta = (DisplayName = "Motion Matching"))
 struct FMotionMatchingTraitSharedData : public FAnimNextTraitSharedData
 {
@@ -99,6 +101,10 @@ struct FMotionMatchingTraitSharedData : public FAnimNextTraitSharedData
 	UPROPERTY(EditAnywhere, Category = Settings)
 	FUAFWeakSystemReference MotionMatchingResultSystemReferenceAlt;
 
+	// Optional blend profile to use when new graph request is made
+	UPROPERTY(EditAnywhere, Category = Settings, meta = (PinHiddenByDefault))
+	TObjectPtr<UUAFBlendProfile> BlendProfile;
+
 	UPROPERTY(EditAnywhere, Category = Interaction, meta = (PinHiddenByDefault))
 	TArray<FPoseSearchInteractionAvailability> Availabilities;
 
@@ -130,6 +136,7 @@ struct FMotionMatchingTraitSharedData : public FAnimNextTraitSharedData
 	UPROPERTY(EditAnywhere, Category = Interaction, meta = (PinHiddenByDefault))
 	FName WarpingTranslationCurveName;
 
+
 	// Latent pin support boilerplate
 	#define TRAIT_LATENT_PROPERTIES_ENUMERATOR(GeneratorMacro) \
 		GeneratorMacro(Databases) \
@@ -153,6 +160,7 @@ struct FMotionMatchingTraitSharedData : public FAnimNextTraitSharedData
 		GeneratorMacro(MotionMatchingResultSystemReference) \
 		GeneratorMacro(MotionMatchingResultVariableNameAlt) \
 		GeneratorMacro(MotionMatchingResultSystemReferenceAlt) \
+		GeneratorMacro(BlendProfile) \
 		GeneratorMacro(Availabilities) \
 		GeneratorMacro(bValidateResultAgainstAvailabilities) \
 		GeneratorMacro(bKeepInteractionAlive) \
diff --git a/Engine/Plugins/Experimental/UAF/UAFPoseSearch/Source/UAFPoseSearch/Private/MotionMatchingTrait.cpp b/Engine/Plugins/Experimental/UAF/UAFPoseSearch/Source/UAFPoseSearch/Private/MotionMatchingTrait.cpp
--- a/Engine/Plugins/Experimental/UAF/UAFPoseSearch/Source/UAFPoseSearch/Private/MotionMatchingTrait.cpp
+++ b/Engine/Plugins/Experimental/UAF/UAFPoseSearch/Source/UAFPoseSearch/Private/MotionMatchingTrait.cpp
@@ -57,6 +57,7 @@ PRAGMA_DISABLE_DEPRECATION_WARNINGS
 #define TRAIT_INTERFACE_ENUMERATOR(GeneratorMacro) \
 		GeneratorMacro(IUpdate) \
 		GeneratorMacro(IEvaluate) \
+		GeneratorMacro(IGarbageCollection) \
 
 	// Trait implementation boilerplate
 	GENERATE_ANIM_TRAIT_IMPLEMENTATION(FMotionMatchingTrait, TRAIT_INTERFACE_ENUMERATOR, NULL_ANIM_TRAIT_INTERFACE_ENUMERATOR, NULL_ANIM_TRAIT_EVENT_ENUMERATOR)
@@ -387,6 +388,8 @@ PRAGMA_ENABLE_DEPRECATION_WARNINGS
 						NewGraphRequest.BlendArgs = SharedData->BlendArguments;
 						NewGraphRequest.AnimationGraph = AnimationGraph;
 						NewGraphRequest.FactoryObject = AnimationAsset;
+						NewGraphRequest.BlendProfile = SharedData->BlendProfile;
+
 						BlendStackBinding.PushGraph(Context, MoveTemp(NewGraphRequest));
 					}
 				}
diff --git a/Engine/Plugins/Experimental/UAF/UAFStateTree/Source/UAFStateTree/Internal/Tasks/AnimNextStateTreeGraphInstanceTask.h b/Engine/Plugins/Experimental/UAF/UAFStateTree/Source/UAFStateTree/Internal/Tasks/AnimNextStateTreeGraphInstanceTask.h
--- a/Engine/Plugins/Experimental/UAF/UAFStateTree/Source/UAFStateTree/Internal/Tasks/AnimNextStateTreeGraphInstanceTask.h
+++ b/Engine/Plugins/Experimental/UAF/UAFStateTree/Source/UAFStateTree/Internal/Tasks/AnimNextStateTreeGraphInstanceTask.h
@@ -13,6 +13,7 @@
 #include "AnimNextStateTreeGraphInstanceTask.generated.h"
 
 struct FAnimNextStateTreeTraitContext;
+class UUAFBlendProfile;
 
 USTRUCT()
 struct UAFSTATETREE_API FAnimNextGraphInstanceTaskInstanceData
@@ -62,6 +63,10 @@ struct UAFSTATETREE_API FAnimNextGraphInstanceTaskInstanceData
 	UPROPERTY(EditAnywhere, Category = Animation)
 	FAlphaBlendArgs BlendOptions;
 
+	// Optional blend profile to use when the state is pushed
+	UPROPERTY(EditAnywhere, Category = Animation)
+	TObjectPtr<UUAFBlendProfile> BlendProfile;
+
 	// Whether this task should continue to tick once state is entered
 	UPROPERTY(EditAnywhere, Category = Animation)
 	bool bContinueTicking = true;
diff --git a/Engine/Plugins/Experimental/UAF/UAFStateTree/Source/UAFStateTree/Private/AnimNextStateTreeContext.cpp b/Engine/Plugins/Experimental/UAF/UAFStateTree/Source/UAFStateTree/Private/AnimNextStateTreeContext.cpp
--- a/Engine/Plugins/Experimental/UAF/UAFStateTree/Source/UAFStateTree/Private/AnimNextStateTreeContext.cpp
+++ b/Engine/Plugins/Experimental/UAF/UAFStateTree/Source/UAFStateTree/Private/AnimNextStateTreeContext.cpp
@@ -8,8 +8,10 @@
 #include "TraitCore/TraitStackBinding.h"
 #include "TraitInterfaces/IBlendStack.h"
 #include "TraitInterfaces/ITimeline.h"
+#include "UAF/BlendProfile/UAFBlendProfile.h"
+#include "Traits/BlendSmootherPerBoneTraitData.h"
 
-bool FAnimNextStateTreeTraitContext::PushAssetOntoBlendStack(UE::UAF::FGraphAssetHandleConstView InAsset, const FAlphaBlendArgs& InBlendArguments) const
+bool FAnimNextStateTreeTraitContext::PushAssetOntoBlendStack(UE::UAF::FGraphAssetHandleConstView InAsset, const FAlphaBlendArgs& InBlendArguments, const UUAFBlendProfile* InBlendProfile) const
 {
 	using namespace UE::UAF;
 
@@ -35,6 +37,7 @@ bool FAnimNextStateTreeTraitContext::PushAssetOntoBlendStack(UE::UAF::FGraphAsse
 	NewGraphRequest.FactoryObject = Objects.Num() > 0 ? Objects[0] : nullptr;
 	NewGraphRequest.AnimationGraph = AnimationGraph;
 	NewGraphRequest.FactoryParams = Params;
+	NewGraphRequest.BlendProfile = InBlendProfile;
 
 	BlendStackBinding.PushGraph(*Context, MoveTemp(NewGraphRequest));
 
diff --git a/Engine/Plugins/Experimental/UAF/UAFStateTree/Source/UAFStateTree/Private/AnimNextStateTreeContext.h b/Engine/Plugins/Experimental/UAF/UAFStateTree/Source/UAFStateTree/Private/AnimNextStateTreeContext.h
--- a/Engine/Plugins/Experimental/UAF/UAFStateTree/Source/UAFStateTree/Private/AnimNextStateTreeContext.h
+++ b/Engine/Plugins/Experimental/UAF/UAFStateTree/Source/UAFStateTree/Private/AnimNextStateTreeContext.h
@@ -16,6 +16,7 @@ struct FStateTreeTrait;
 }
 
 class UUAFAnimGraph;
+class UUAFBlendProfile;
 struct FAlphaBlendArgs;
 struct FAnimNextFactoryParams;
 
@@ -34,7 +35,7 @@ struct FAnimNextStateTreeTraitContext
 		, Binding(&InBinding)
 	{}
 
-	bool PushAssetOntoBlendStack(UE::UAF::FGraphAssetHandleConstView InAsset, const FAlphaBlendArgs& InBlendArguments) const;
+	bool PushAssetOntoBlendStack(UE::UAF::FGraphAssetHandleConstView InAsset, const FAlphaBlendArgs& InBlendArguments, const UUAFBlendProfile* InBlendProfile) const;
 	float QueryPlaybackRatio() const;
 	float QueryTimeLeft() const;
 	float QueryDuration() const;
diff --git a/Engine/Plugins/Experimental/UAF/UAFStateTree/Source/UAFStateTree/Private/Tasks/AnimNextStateTreeGraphInstanceTask.cpp b/Engine/Plugins/Experimental/UAF/UAFStateTree/Source/UAFStateTree/Private/Tasks/AnimNextStateTreeGraphInstanceTask.cpp
--- a/Engine/Plugins/Experimental/UAF/UAFStateTree/Source/UAFStateTree/Private/Tasks/AnimNextStateTreeGraphInstanceTask.cpp
+++ b/Engine/Plugins/Experimental/UAF/UAFStateTree/Source/UAFStateTree/Private/Tasks/AnimNextStateTreeGraphInstanceTask.cpp
@@ -175,7 +175,7 @@ EStateTreeRunStatus FAnimNextStateTreeGraphInstanceTask::EnterState(FStateTreeEx
 	FAnimNextStateTreeTraitContext& ExecContext = Context.GetExternalData(TraitContextHandle);
 	FInstanceDataType& InstanceData = Context.GetInstanceData(*this);
 
-	if (ExecContext.PushAssetOntoBlendStack(InstanceData.AssetData, InstanceData.BlendOptions))
+	if (ExecContext.PushAssetOntoBlendStack(InstanceData.AssetData, InstanceData.BlendOptions, InstanceData.BlendProfile))
 	{
 		return EStateTreeRunStatus::Running;
 	}
diff --git a/Engine/Plugins/Experimental/UAF/UAFStateTree/Source/UAFStateTree/UAFStateTree.Build.cs b/Engine/Plugins/Experimental/UAF/UAFStateTree/Source/UAFStateTree/UAFStateTree.Build.cs
--- a/Engine/Plugins/Experimental/UAF/UAFStateTree/Source/UAFStateTree/UAFStateTree.Build.cs
+++ b/Engine/Plugins/Experimental/UAF/UAFStateTree/Source/UAFStateTree/UAFStateTree.Build.cs
@@ -21,6 +21,8 @@ namespace UnrealBuildTool.Rules
 					"StateTreeModule",
 					"Engine",
 					"UAF",
+					"HierarchyTableRuntime",
+					"HierarchyTableAnimationRuntime",
 				}
 			);
 
diff --git a/Engine/Plugins/Experimental/UAF/UAFStateTree/Source/UAFStateTreeUncookedOnly/Private/AnimNextStateTree_EditorData.cpp b/Engine/Plugins/Experimental/UAF/UAFStateTree/Source/UAFStateTreeUncookedOnly/Private/AnimNextStateTree_EditorData.cpp
--- a/Engine/Plugins/Experimental/UAF/UAFStateTree/Source/UAFStateTreeUncookedOnly/Private/AnimNextStateTree_EditorData.cpp
+++ b/Engine/Plugins/Experimental/UAF/UAFStateTree/Source/UAFStateTreeUncookedOnly/Private/AnimNextStateTree_EditorData.cpp
@@ -32,6 +32,7 @@
 #include "Entries/AnimNextSharedVariablesEntry.h"
 #include "Traits/BlendStackTrait.h"
 #include "Traits/BlendSmoother.h"
+#include "Traits/BlendSmootherPerBoneTraitData.h"
 
 TSubclassOf<UAssetUserData> UAnimNextStateTree_EditorData::GetAssetUserDataClass() const
 {
@@ -213,6 +214,14 @@ void UAnimNextStateTree_EditorData::OnPreCompileGetProgrammaticGraphs(const FRig
 			InSettings.ReportError(TEXT("Failed to add Blend Smoother Core trait to node"));
 			return;
 		}
+					
+		const FName SmootherPerBoneTraitName = Controller->AddTraitStruct(TraitStackNode, TInstancedStruct<FAnimNextBlendSmootherPerBoneCoreTraitSharedData>::Make(), INDEX_NONE, false, false);
+
+		if (SmootherPerBoneTraitName == NAME_None)
+		{
+			InSettings.ReportError(TEXT("Failed to add Blend Smoother Per Bone Core trait to node"));
+			return;
+		}
 
 		URigVMPin* StateTreeReferencePin = TraitStackNode->FindTrait(StateTreeTraitName, GET_MEMBER_NAME_STRING_CHECKED(FAnimNextStateTreeTraitSharedData, StateTreeReference));
 		if (StateTreeReferencePin == nullptr)
diff --git a/Engine/Source/Runtime/Engine/Classes/Animation/AnimationAsset.h b/Engine/Source/Runtime/Engine/Classes/Animation/AnimationAsset.h
--- a/Engine/Source/Runtime/Engine/Classes/Animation/AnimationAsset.h
+++ b/Engine/Source/Runtime/Engine/Classes/Animation/AnimationAsset.h
@@ -202,6 +202,7 @@ struct FBlendSampleData
 	}
 
 	static void ENGINE_API NormalizeDataWeight(TArray<FBlendSampleData>& SampleDataList);
+	static void ENGINE_API NormalizeDataWeight(const TArrayView<FBlendSampleData*>& SampleDataList);
 };
 
 USTRUCT()
diff --git a/Engine/Source/Runtime/Engine/Private/AnimationAsset.cpp b/Engine/Source/Runtime/Engine/Private/AnimationAsset.cpp
--- a/Engine/Source/Runtime/Engine/Private/AnimationAsset.cpp
+++ b/Engine/Source/Runtime/Engine/Private/AnimationAsset.cpp
@@ -873,4 +873,73 @@ void FBlendSampleData::NormalizeDataWeight(TArray<FBlendSampleData>& SampleDataL
 	}
 }
 
+void FBlendSampleData::NormalizeDataWeight(const TArrayView<FBlendSampleData*>& SampleDataList)
+{
+	float TotalSum = 0.f;
+
+	check(SampleDataList.Num() > 0);
+	const int32 NumBones = SampleDataList[0]->PerBoneBlendData.Num();
+
+	FMemMark Mark(FMemStack::Get());
+	TArray<float, TMemStackAllocator<>> PerBoneTotalSums;
+	PerBoneTotalSums.AddZeroed(NumBones);
+
+	for (int32 PoseIndex = 0; PoseIndex < SampleDataList.Num(); PoseIndex++)
+	{
+		checkf(SampleDataList[PoseIndex]->PerBoneBlendData.Num() == NumBones, TEXT("Attempted to normalise a blend sample list, but the samples have differing numbers of bones."));
+
+		TotalSum += SampleDataList[PoseIndex]->TotalWeight;
+
+		if (SampleDataList[PoseIndex]->PerBoneBlendData.Num() > 0)
+		{
+			// now interpolate the per bone weights
+			for (int32 BoneIndex = 0; BoneIndex < NumBones; BoneIndex++)
+			{
+				PerBoneTotalSums[BoneIndex] += SampleDataList[PoseIndex]->PerBoneBlendData[BoneIndex];
+			}
+		}
+	}
+
+	// Re-normalize Pose weight
+	if (TotalSum > ZERO_ANIMWEIGHT_THRESH)
+	{
+		if (FMath::Abs<float>(TotalSum - 1.f) > ZERO_ANIMWEIGHT_THRESH)
+		{
+			for (int32 PoseIndex = 0; PoseIndex < SampleDataList.Num(); PoseIndex++)
+			{
+				SampleDataList[PoseIndex]->TotalWeight /= TotalSum;
+			}
+		}
+	}
+	else
+	{
+		for (int32 PoseIndex = 0; PoseIndex < SampleDataList.Num(); PoseIndex++)
+		{
+			SampleDataList[PoseIndex]->TotalWeight = 1.0f / SampleDataList.Num();
+		}
+	}
+
+	// Re-normalize per bone weights.
+	for (int32 BoneIndex = 0; BoneIndex < NumBones; BoneIndex++)
+	{
+		if (PerBoneTotalSums[BoneIndex] > ZERO_ANIMWEIGHT_THRESH)
+		{
+			if (FMath::Abs<float>(PerBoneTotalSums[BoneIndex] - 1.f) > ZERO_ANIMWEIGHT_THRESH)
+			{
+				for (int32 PoseIndex = 0; PoseIndex < SampleDataList.Num(); PoseIndex++)
+				{
+					SampleDataList[PoseIndex]->PerBoneBlendData[BoneIndex] /= PerBoneTotalSums[BoneIndex];
+				}
+			}
+		}
+		else
+		{
+			for (int32 PoseIndex = 0; PoseIndex < SampleDataList.Num(); PoseIndex++)
+			{
+				SampleDataList[PoseIndex]->PerBoneBlendData[BoneIndex] = 1.0f / SampleDataList.Num();
+			}
+		}
+	}
+}
+
 #undef LOCTEXT_NAMESPACE // "AnimationAsset"
