UAF: Fix GC crash in shipping builds

Change check() to ensure() so that shipping builds actually execute the code

#Synced-CL 51429293

diff --git a/Engine/Plugins/Experimental/UAF/UAFAnimGraph/Source/UAFAnimGraph/Private/Graph/AnimNextGraphInstance.cpp b/Engine/Plugins/Experimental/UAF/UAFAnimGraph/Source/UAFAnimGraph/Private/Graph/AnimNextGraphInstance.cpp
--- a/Engine/Plugins/Experimental/UAF/UAFAnimGraph/Source/UAFAnimGraph/Private/Graph/AnimNextGraphInstance.cpp
+++ b/Engine/Plugins/Experimental/UAF/UAFAnimGraph/Source/UAFAnimGraph/Private/Graph/AnimNextGraphInstance.cpp
@@ -374,9 +374,10 @@ void FAnimNextGraphInstance::FGCHandler::AddReferencedObjects(FReferenceCollecto
 
 		if (Context.GetStack(TraitPtr, TraitStack))
 		{
-			check(TraitStack.GetInterface(GCTrait));
-
-			GCTrait.AddReferencedObjects(Context, Collector);
+			if(ensure(TraitStack.GetInterface(GCTrait)))
+			{
+				GCTrait.AddReferencedObjects(Context, Collector);
+			}
 		}
 	}
 }
