UAF: Add GetSlotLocalWeight function and RigUnit

#Synced-CL 51454554

diff --git a/Engine/Plugins/Experimental/UAF/UAFAnimGraph/Source/UAFAnimGraph/Internal/Injection/UAFMontageComponent.h b/Engine/Plugins/Experimental/UAF/UAFAnimGraph/Source/UAFAnimGraph/Internal/Injection/UAFMontageComponent.h
--- a/Engine/Plugins/Experimental/UAF/UAFAnimGraph/Source/UAFAnimGraph/Internal/Injection/UAFMontageComponent.h
+++ b/Engine/Plugins/Experimental/UAF/UAFAnimGraph/Source/UAFAnimGraph/Internal/Injection/UAFMontageComponent.h
@@ -51,9 +51,12 @@ public:
 	// Returns true if the montage in this eval data is active (valid and not blending out)
 	bool IsActive() const { return (Montage != nullptr && BlendInfo.GetDesiredValue() != 0.0f); }
 	
-	// Return the montage for this eval data
+	// Returns the montage for this eval data
 	const UAnimMontage* GetMontage() const {return Montage;}
 	
+	// Returns the alpha blend info for this eval data
+	const FAlphaBlend& GetBlendInfo() const {return BlendInfo;}
+	
 private:
 	UPROPERTY(Transient)
 	TObjectPtr<UAnimMontage> Montage = nullptr;
@@ -84,6 +87,7 @@ public:
 	virtual ~FUAFMontageComponent() override;
 
 	UE_API bool IsSlotActive(FName SlotNodeName) const;
+	UE_API float GetSlotLocalWeight(FName SlotNodeName) const;
 	UE_API TObjectPtr<const UAnimMontage> GetActiveMontageForSlot(FName SlotNodeName) const;
 	
 private:
@@ -115,7 +119,7 @@ private:
 	TSubclassOf<UAnimInstance> LegacyMontageAnimInstanceClass = nullptr;
 
 	// The event name of the module tick function used at runtime to enqueue the montage tick function as a prerequisite to. 
-	// If the the exposed TickMontagesBeforeEventName does not provide an override, it will fall back to DefaultTickDependencyEventName
+	// If the exposed TickMontagesBeforeEventName does not provide an override, it will fall back to DefaultTickDependencyEventName
 	UPROPERTY(Transient)
 	FName TickDependencyEventName = NAME_None;
 	
@@ -137,7 +141,7 @@ struct TStructOpsTypeTraits<FUAFMontageComponent> : public TStructOpsTypeTraitsB
 	};
 };
 
-// Return true if the running UAF System has an active montage in the given slot. A UAnimMontage that is playing in the slot and blending out is not determined to be "active". 
+// Returns true if the running UAF System has an active montage in the given slot. A UAnimMontage that is playing in the slot and blending out is not determined to be "active". 
 USTRUCT(meta=(DisplayName="Is Slot Active", Category="Injection", NodeColor="0, 1, 1", Keywords="Slot, Active, Montage"))
 struct FRigUnit_IsSlotActive : public FRigUnit_AnimNextBase
 {
@@ -150,9 +154,27 @@ struct FRigUnit_IsSlotActive : public FRigUnit_AnimNextBase
 	UPROPERTY(EditAnywhere, Category = "Graph", meta = (Input))
 	FName SlotNodeName = NAME_None;
 	
-	// True if the running UAF System has an active montage in the given slot and it's not blending out
+	// True if the running UAF System has an active montage in the given slot, and it's not blending out
 	UPROPERTY(EditAnywhere, Category = "Graph", meta = (Output))
 	bool Result = false;
 };
 
+// Gets the local weight of any montages this slot node is playing. If this slot is not currently playing a montage, it will return 0. 
+USTRUCT(meta=(DisplayName="Get Slot Local Weight", Category="Injection", NodeColor="0, 1, 1", Keywords="Slot, Weight, Montage"))
+struct FRigUnit_GetSlotLocalWeight : public FRigUnit_AnimNextBase
+{
+	GENERATED_BODY()
+
+	RIGVM_METHOD()
+	void Execute();
+	
+	// The name of the slot
+	UPROPERTY(EditAnywhere, Category = "Graph", meta = (Input))
+	FName SlotNodeName = NAME_None;
+	
+	// The local weight of any montages this slot is playing
+	UPROPERTY(EditAnywhere, Category = "Graph", meta = (Output))
+	float LocalWeight = 0.0f;
+};
+
 #undef UE_API
diff --git a/Engine/Plugins/Experimental/UAF/UAFAnimGraph/Source/UAFAnimGraph/Private/Injection/UAFMontageComponent.cpp b/Engine/Plugins/Experimental/UAF/UAFAnimGraph/Source/UAFAnimGraph/Private/Injection/UAFMontageComponent.cpp
--- a/Engine/Plugins/Experimental/UAF/UAFAnimGraph/Source/UAFAnimGraph/Private/Injection/UAFMontageComponent.cpp
+++ b/Engine/Plugins/Experimental/UAF/UAFAnimGraph/Source/UAFAnimGraph/Private/Injection/UAFMontageComponent.cpp
@@ -127,6 +127,25 @@ bool FUAFMontageComponent::IsSlotActive(FName SlotNodeName) const
 	return false;
 }
 
+float FUAFMontageComponent::GetSlotLocalWeight(FName SlotNodeName) const
+{
+	if (SlotNodeName != NAME_None)
+	{
+		float MontageLocalWeight = 0.0f;
+		for (const FUAFMontageEvalData& EvalData : MontageEvalData)
+		{
+			const UAnimMontage* const Montage = EvalData.GetMontage();
+			if (Montage && Montage->IsValidSlot(SlotNodeName))
+			{
+				const float MontageWeight = EvalData.GetBlendInfo().GetBlendedValue();
+				MontageLocalWeight += MontageWeight;
+			}
+		}
+		return FMath::Clamp(MontageLocalWeight, 0.0f, 1.0f);
+	}
+	return 0.0f;
+}
+
 TObjectPtr<const UAnimMontage> FUAFMontageComponent::GetActiveMontageForSlot(FName SlotNodeName) const
 {
 	if (SlotNodeName != NAME_None)
@@ -277,4 +296,21 @@ FRigUnit_IsSlotActive_Execute()
 	}
 }
 
+FRigUnit_GetSlotLocalWeight_Execute()
+{
+	LocalWeight = 0.0f;
+	if (SlotNodeName == NAME_None)
+	{
+		UAF_RIGUNIT_LOG(Warning, TEXT("[%s] is not a valid slot name."), *SlotNodeName.ToString());
+		return;
+	}
+	
+	FUAFAssetInstance& Instance = ExecuteContext.GetContextData<FUAFAssetContextData>().GetInstance();
+	FAnimNextModuleInstance* RootInstance = Instance.GetRootInstance();
+	const FUAFMontageComponent* MontageComponent = RootInstance ? RootInstance->TryGetComponent<FUAFMontageComponent>() : nullptr;
+	if (MontageComponent)
+	{
+		LocalWeight = MontageComponent->GetSlotLocalWeight(SlotNodeName);
+	}
+}
 
