UAF: Fix GC for graph instances

FGCObject now used as the root of the GC for a graph instance, not the struct ARO

#Synced-CL 51493025

diff --git a/Engine/Plugins/Experimental/UAF/UAFAnimGraph/Source/UAFAnimGraph/Private/Graph/AnimNextGraphInstance.cpp b/Engine/Plugins/Experimental/UAF/UAFAnimGraph/Source/UAFAnimGraph/Private/Graph/AnimNextGraphInstance.cpp
--- a/Engine/Plugins/Experimental/UAF/UAFAnimGraph/Source/UAFAnimGraph/Private/Graph/AnimNextGraphInstance.cpp
+++ b/Engine/Plugins/Experimental/UAF/UAFAnimGraph/Source/UAFAnimGraph/Private/Graph/AnimNextGraphInstance.cpp
@@ -19,6 +19,7 @@
 #include "Module/AnimNextModule.h"
 #include "Module/AnimNextModuleInstance.h"
 #include "Module/AnimNextSkeletalMeshComponentReferenceComponent.h"
+#include "TraitCore/NodeInstance.h"
 #include "TraitInterfaces/IGarbageCollection.h"
 
 #include UE_INLINE_GENERATED_CPP_BY_NAME(AnimNextGraphInstance)
@@ -27,6 +28,7 @@ DEFINE_STAT(STAT_AnimNext_Graph_RigVM);
 
 FAnimNextGraphInstance::FAnimNextGraphInstance()
 	: FUAFAssetInstance(StaticStruct())
+	, GCHandler(*this)
 {
 }
 
@@ -354,15 +356,12 @@ void FAnimNextGraphInstance::RunOrQueueTask(UE::UAF::FUniqueAnimGraphTask&& InTa
 	}
 }
 
-void FAnimNextGraphInstance::AddStructReferencedObjects(FReferenceCollector& Collector)
-{
-	GCHandler.AddReferencedObjects(Collector);
-}
-
 void FAnimNextGraphInstance::FGCHandler::AddReferencedObjects(FReferenceCollector& Collector)
 {
 	using namespace UE::UAF;
 
+	Collector.AddPropertyReferencesWithStructARO(FAnimNextGraphInstance::StaticStruct(), &Owner);
+
 	FExecutionContext Context;
 	FTraitStackBinding TraitStack;
 	TTraitBinding<IGarbageCollection> GCTrait;
diff --git a/Engine/Plugins/Experimental/UAF/UAFAnimGraph/Source/UAFAnimGraph/Public/Graph/AnimNextGraphInstance.h b/Engine/Plugins/Experimental/UAF/UAFAnimGraph/Source/UAFAnimGraph/Public/Graph/AnimNextGraphInstance.h
--- a/Engine/Plugins/Experimental/UAF/UAFAnimGraph/Source/UAFAnimGraph/Public/Graph/AnimNextGraphInstance.h
+++ b/Engine/Plugins/Experimental/UAF/UAFAnimGraph/Source/UAFAnimGraph/Public/Graph/AnimNextGraphInstance.h
@@ -112,8 +112,6 @@ struct FAnimNextGraphInstance : public FUAFAssetInstance
 	// Called each time the graph updates to mark the instance as updated
 	UE_API void MarkAsUpdated();
 
-	void AddStructReferencedObjects(FReferenceCollector& Collector);
-
 private:
 	// Executes a list of latent RigVM pins and writes the result into the destination pointer (latent handle offsets are using the destination as base)
 	// When frozen, latent handles that can freeze are skipped, all others will execute
@@ -152,10 +150,17 @@ private:
 	// GC handler. This is nested to avoid having to bubble the polymorphism it introduces up to super-structs
 	struct FGCHandler : public FGCObject
 	{
+		explicit FGCHandler(FAnimNextGraphInstance& InOwner)
+			: Owner(InOwner)
+		{}
+
 		// FGCObject implementation
 		virtual void AddReferencedObjects(FReferenceCollector& Collector) override;
 		virtual FString GetReferencerName() const override;
 
+		// Owner instance
+		FAnimNextGraphInstance& Owner;
+
 		// List of trait handles that contain UObject references and implement IGarbageCollection
 		TArray<UE::UAF::FWeakTraitPtr> TraitsWithReferences;
 	} GCHandler;
@@ -216,7 +221,6 @@ struct TStructOpsTypeTraits<FAnimNextGraphInstance> : public TStructOpsTypeTrait
 	enum
 	{
 		WithCopy = false,
-		WithAddStructReferencedObjects = true,
 	};
 };
 
