Move UAF notify handling into a common location

#Synced-CL 51544635

diff --git a/Engine/Plugins/Experimental/UAF/UAF/Source/UAF/Internal/Module/AnimNextModuleInstance.h b/Engine/Plugins/Experimental/UAF/UAF/Source/UAF/Internal/Module/AnimNextModuleInstance.h
--- a/Engine/Plugins/Experimental/UAF/UAF/Source/UAF/Internal/Module/AnimNextModuleInstance.h
+++ b/Engine/Plugins/Experimental/UAF/UAF/Source/UAF/Internal/Module/AnimNextModuleInstance.h
@@ -27,7 +27,7 @@ struct FRigUnit_UAFRunAsset;
 class UUAFEngineSubsystem;
 struct FRigUnit_AnimNextInitializeEvent;
 struct FRigUnit_CopyModuleProxyVariables;
-struct FAnimNextNotifyDispatcherComponent;
+struct FUAFNotifyDispatcherComponent;
 struct FAnimNextGraphInstance;
 struct FUAFSystemFactoryParams;
 struct FUAFWeakSystemReference;
@@ -216,7 +216,7 @@ private:
 	friend UUAFEngineSubsystem;
 	friend FRigUnit_AnimNextInitializeEvent;
 	friend FRigUnit_CopyModuleProxyVariables;
-	friend FAnimNextNotifyDispatcherComponent;
+	friend FUAFNotifyDispatcherComponent;
 	friend FAnimNextGraphInstance;
 	friend UE::UAF::Tests::FVariables_Shared;
 	friend UE::UAF::Tests::FVariables_UOLBindings;
diff --git a/Engine/Plugins/Experimental/UAF/UAF/Source/UAF/Private/Module/UAFNotifyDispatcherComponent.cpp b/Engine/Plugins/Experimental/UAF/UAF/Source/UAF/Private/Module/UAFNotifyDispatcherComponent.cpp
new file mode 100644
--- /dev/null
+++ b/Engine/Plugins/Experimental/UAF/UAF/Source/UAF/Private/Module/UAFNotifyDispatcherComponent.cpp
@@ -0,0 +1,246 @@
+// Copyright Epic Games, Inc. All Rights Reserved.
+
+#include "Module/UAFNotifyDispatcherComponent.h"
+
+#include "Animation/AnimSequenceBase.h"
+#include "Animation/AnimNotifies/AnimNotify.h"
+#include "Animation/AnimNotifies/AnimNotifyState.h"
+#include "Module/AnimNextModuleInstance.h"
+#include "Module/AnimNextSkeletalMeshComponentReferenceComponent.h"
+
+#include UE_INLINE_GENERATED_CPP_BY_NAME(UAFNotifyDispatcherComponent)
+
+void FUAFNotifyDispatcherComponent::OnBindToInstance()
+{
+	// Default to the current host's skeletal mesh component, if any
+	FAnimNextModuleInstance& ModuleInstance = GetModuleInstance();
+	const FAnimNextSkeletalMeshComponentReferenceComponent& ComponentReference = ModuleInstance.GetComponent<FAnimNextSkeletalMeshComponentReferenceComponent>();
+	SkeletalMeshComponent = ComponentReference.GetComponent();
+	NotifyQueue.PredictedLODLevel = SkeletalMeshComponent ? SkeletalMeshComponent->GetPredictedLODLevel() : 0;
+}
+
+void FUAFNotifyDispatcherComponent::OnTraitEvent(FAnimNextTraitEvent& Event)
+{
+	if (UE::UAF::FNotifyDispatchEvent* NotifyDispatchEvent = Event.AsType<UE::UAF::FNotifyDispatchEvent>())
+	{
+		NotifyQueue.AddAnimNotifies(NotifyDispatchEvent->Notifies, NotifyDispatchEvent->Weight);
+		NotifyDispatchEvent->MarkConsumed();
+	}
+}
+
+void FUAFNotifyDispatcherComponent::OnEndExecution(float InDeltaTime)
+{
+	using namespace UE::UAF;
+
+	// Early out if we dont need to do any work
+	if (NotifyQueue.AnimNotifies.Num() == 0 && ActiveAnimNotifyState.Num() == 0)
+	{
+		return;
+	}
+
+	FAnimNextModuleInstance& ModuleInstance = GetModuleInstance();
+
+	TSharedPtr<FNotifyQueueDispatchEvent> NotifyQueueDispatchEvent;
+	if (SkeletalMeshComponent != nullptr)
+	{
+		NotifyQueueDispatchEvent = MakeTraitEvent<FNotifyQueueDispatchEvent>();
+		NotifyQueueDispatchEvent->WeakSkeletalMeshComponent = SkeletalMeshComponent;
+		NotifyQueueDispatchEvent->DeltaSeconds = InDeltaTime;
+	}
+
+	NotifyQueue.PredictedLODLevel = SkeletalMeshComponent ? SkeletalMeshComponent->GetPredictedLODLevel() : 0;
+
+	// Logic from here largely mimics that of UAnimInstance::TriggerAnimNotifies 
+
+	// Array that will replace the 'ActiveAnimNotifyState' at the end of this function.
+	TArray<const FAnimNotifyEvent*> NewActiveAnimNotifyState;
+	NewActiveAnimNotifyState.Reserve(NotifyQueue.AnimNotifies.Num());
+
+	TArray<FAnimNotifyEventReference> NewActiveAnimNotifyEventReference;
+	NewActiveAnimNotifyEventReference.Reserve(NotifyQueue.AnimNotifies.Num());
+
+	// AnimNotifyState freshly added that need their 'NotifyBegin' event called.
+	TArray<const FAnimNotifyEvent*> NotifyStateBeginEvent;
+	TArray<const FAnimNotifyEventReference*> NotifyStateBeginEventReference;
+
+	for (int32 Index = 0; Index < NotifyQueue.AnimNotifies.Num(); Index++)
+	{
+		if (const FAnimNotifyEvent* AnimNotifyEvent = NotifyQueue.AnimNotifies[Index].GetNotify())
+		{
+			// AnimNotifyState
+			if (AnimNotifyEvent->NotifyStateClass)
+			{
+				int32 ExistingItemIndex = ActiveAnimNotifyState.IndexOfByPredicate([AnimNotifyEvent](const FAnimNotifyEvent* InOther)
+					{
+						return *AnimNotifyEvent == *InOther;
+					});
+
+				if (ExistingItemIndex != INDEX_NONE)
+				{
+					check(ActiveAnimNotifyState.Num() == ActiveAnimNotifyEventReference.Num());
+					ActiveAnimNotifyState.RemoveAtSwap(ExistingItemIndex, EAllowShrinking::No);
+					ActiveAnimNotifyEventReference.RemoveAtSwap(ExistingItemIndex, EAllowShrinking::No);
+				}
+				else
+				{
+					NotifyStateBeginEvent.Add(AnimNotifyEvent);
+					NotifyStateBeginEventReference.Add(&NotifyQueue.AnimNotifies[Index]);
+				}
+
+				NewActiveAnimNotifyState.Add(AnimNotifyEvent);
+				FAnimNotifyEventReference& EventRef = NewActiveAnimNotifyEventReference.Add_GetRef(NotifyQueue.AnimNotifies[Index]);
+				EventRef.SetNotify(AnimNotifyEvent);
+				continue;
+			}
+
+			// Trigger non 'state' AnimNotifies
+			TriggerSingleAnimNotify(InDeltaTime, NotifyQueueDispatchEvent.Get(), NotifyQueue.AnimNotifies[Index]);
+		}
+	}
+
+	if (NotifyQueueDispatchEvent.IsValid())
+	{
+		// Send end notification to AnimNotifyState not active anymore.
+		for (int32 Index = 0; Index < ActiveAnimNotifyState.Num(); ++Index)
+		{
+			const FAnimNotifyEvent* AnimNotifyEvent = ActiveAnimNotifyState[Index];
+			const FAnimNotifyEventReference& EventReference = ActiveAnimNotifyEventReference[Index];
+			if (AnimNotifyEvent->NotifyStateClass)
+			{
+#if WITH_EDITOR
+				// Prevent firing notifies in animation editors if requested
+				if (ModuleInstance.GetWorldType() != EWorldType::EditorPreview || AnimNotifyEvent->NotifyStateClass->ShouldFireInEditor())
+#endif
+				{
+					NotifyQueueDispatchEvent->EventsToEnd.Add(EventReference);
+				}
+			}
+		}
+
+		check(NotifyStateBeginEventReference.Num() == NotifyStateBeginEvent.Num());
+		for (int32 Index = 0; Index < NotifyStateBeginEvent.Num(); Index++)
+		{
+			const FAnimNotifyEvent* AnimNotifyEvent = NotifyStateBeginEvent[Index];
+			const FAnimNotifyEventReference* EventReference = NotifyStateBeginEventReference[Index];
+			if (AnimNotifyEvent->NotifyStateClass)
+			{
+#if WITH_EDITOR
+				// Prevent firing notifies in animation editors if requested 
+				if (ModuleInstance.GetWorldType() != EWorldType::EditorPreview || AnimNotifyEvent->NotifyStateClass->ShouldFireInEditor())
+#endif
+				{
+					NotifyQueueDispatchEvent->EventsToBegin.Add(*EventReference);
+				}
+			}
+		}
+	}
+
+	// Switch our arrays.
+	ActiveAnimNotifyState = MoveTemp(NewActiveAnimNotifyState);
+	ActiveAnimNotifyEventReference = MoveTemp(NewActiveAnimNotifyEventReference);
+
+	if (NotifyQueueDispatchEvent.IsValid())
+	{
+		// Tick currently active AnimNotifyState
+		for (int32 Index = 0; Index < ActiveAnimNotifyState.Num(); Index++)
+		{
+			const FAnimNotifyEvent* AnimNotifyEvent = ActiveAnimNotifyState[Index];
+			const FAnimNotifyEventReference& EventReference = ActiveAnimNotifyEventReference[Index];
+			if (AnimNotifyEvent->NotifyStateClass)
+			{
+#if WITH_EDITOR
+				// Prevent firing notifies in animation editors if requested 
+				if (ModuleInstance.GetWorldType() != EWorldType::EditorPreview || AnimNotifyEvent->NotifyStateClass->ShouldFireInEditor())
+#endif
+				{
+					NotifyQueueDispatchEvent->EventsToTick.Add(EventReference);
+				}
+			}
+		}
+
+		GetModuleInstance().QueueOutputTraitEvent(NotifyQueueDispatchEvent);
+	}
+	NotifyQueue.AnimNotifies.Reset();
+}
+
+void FUAFNotifyDispatcherComponent::TriggerSingleAnimNotify(float InDeltaTime, UE::UAF::FNotifyQueueDispatchEvent* InDispatcher, const FAnimNotifyEventReference& EventReference)
+{
+	// This is for non 'state' anim notifies.
+	const FAnimNotifyEvent* AnimNotifyEvent = EventReference.GetNotify();
+	if (AnimNotifyEvent && (AnimNotifyEvent->NotifyStateClass == nullptr))
+	{
+		FAnimNextModuleInstance& ModuleInstance = GetModuleInstance();
+		if (InDispatcher != nullptr && AnimNotifyEvent->Notify != nullptr)
+		{
+#if WITH_EDITOR
+			// Prevent firing notifies in animation editors if requested 
+			if (ModuleInstance.GetWorldType() != EWorldType::EditorPreview || AnimNotifyEvent->Notify->ShouldFireInEditor())
+#endif
+			{
+				InDispatcher->EventsToNotify.Add(EventReference);
+			}
+		}
+		else if (AnimNotifyEvent->NotifyName != NAME_None)
+		{
+			// Named notifies can be handled by a custom module event on our worker thread, no need to queue
+			ModuleInstance.RunScriptEventByName(AnimNotifyEvent->NotifyName, InDeltaTime);
+		}
+	}
+}
+
+namespace UE::UAF
+{
+	template<typename TriggerFunc>
+	static void TriggerEventFunc(const TArray<FAnimNotifyEventReference>& InEvents, TriggerFunc InTriggerFunc)
+	{
+		for (const FAnimNotifyEventReference& EventReference : InEvents)
+		{
+			const FAnimNotifyEvent* AnimNotifyEvent = EventReference.GetNotify();
+			if (AnimNotifyEvent == nullptr)
+			{
+				continue;
+			}
+
+			InTriggerFunc(AnimNotifyEvent, EventReference);
+		}
+	}
+
+	void FNotifyQueueDispatchEvent::Execute() const
+	{
+		USkeletalMeshComponent* SkeletalMeshComponent = WeakSkeletalMeshComponent.Get();
+		if (SkeletalMeshComponent == nullptr)
+		{
+			return;
+		}
+
+		// Order of dispatch here mimics that of UAnimInstance::TriggerAnimNotifies
+
+		// Triggers...
+		TriggerEventFunc(EventsToNotify, [SkeletalMeshComponent](const FAnimNotifyEvent* InAnimNotifyEvent, const FAnimNotifyEventReference& InEventReference)
+			{
+				check(InAnimNotifyEvent->Notify);
+				InAnimNotifyEvent->Notify->Notify(SkeletalMeshComponent, Cast<UAnimSequenceBase>(InAnimNotifyEvent->Notify->GetOuter()), InEventReference);
+			});
+
+		// Ends...
+		TriggerEventFunc(EventsToEnd, [SkeletalMeshComponent](const FAnimNotifyEvent* InAnimNotifyEvent, const FAnimNotifyEventReference& InEventReference)
+			{
+				check(InAnimNotifyEvent->NotifyStateClass);
+				InAnimNotifyEvent->NotifyStateClass->NotifyEnd(SkeletalMeshComponent, Cast<UAnimSequenceBase>(InAnimNotifyEvent->NotifyStateClass->GetOuter()), InEventReference);
+			});
+
+		// Begins...
+		TriggerEventFunc(EventsToBegin, [SkeletalMeshComponent](const FAnimNotifyEvent* InAnimNotifyEvent, const FAnimNotifyEventReference& InEventReference)
+			{
+				check(InAnimNotifyEvent->NotifyStateClass);
+				InAnimNotifyEvent->NotifyStateClass->NotifyBegin(SkeletalMeshComponent, Cast<UAnimSequenceBase>(InAnimNotifyEvent->NotifyStateClass->GetOuter()), InAnimNotifyEvent->GetDuration(), InEventReference);
+			});
+
+		// Ticks...
+		TriggerEventFunc(EventsToTick, [SkeletalMeshComponent, DeltaSeconds = DeltaSeconds](const FAnimNotifyEvent* InAnimNotifyEvent, const FAnimNotifyEventReference& InEventReference)
+			{
+				check(InAnimNotifyEvent->NotifyStateClass);
+				InAnimNotifyEvent->NotifyStateClass->NotifyTick(SkeletalMeshComponent, Cast<UAnimSequenceBase>(InAnimNotifyEvent->NotifyStateClass->GetOuter()), DeltaSeconds, InEventReference);
+			});
+	}
+}
diff --git a/Engine/Plugins/Experimental/UAF/UAF/Source/UAF/Public/Module/UAFNotifyDispatcherComponent.h b/Engine/Plugins/Experimental/UAF/UAF/Source/UAF/Public/Module/UAFNotifyDispatcherComponent.h
new file mode 100644
--- /dev/null
+++ b/Engine/Plugins/Experimental/UAF/UAF/Source/UAF/Public/Module/UAFNotifyDispatcherComponent.h
@@ -0,0 +1,91 @@
+// Copyright Epic Games, Inc. All Rights Reserved.
+
+#pragma once
+
+#include "CoreMinimal.h"
+#include "Module/ModuleEvents.h"
+#include "Module/UAFModuleInstanceComponent.h"
+
+#include "UAFNotifyDispatcherComponent.generated.h"
+
+struct FAnimNextTraitEvent;
+struct FAnimNotifyEventReference;
+
+namespace UE::UAF
+{
+	struct FNotifyQueueDispatchEvent;
+}
+
+#define UE_API UAF_API
+
+USTRUCT()
+struct FUAFNotifyDispatcherComponent : public FUAFModuleInstanceComponent
+{
+	GENERATED_BODY()
+	DECLARE_UAF_ASSET_INSTANCE_COMPONENT()
+
+	// FUAFModuleInstanceComponent interface
+	virtual void OnTraitEvent(FAnimNextTraitEvent& Event) override;
+	virtual void OnEndExecution(float InDeltaTime) override;
+
+	// Triggers a single anim notify in the dispatcher
+	void TriggerSingleAnimNotify(float InDeltaTime, UE::UAF::FNotifyQueueDispatchEvent* InDispatcher, const FAnimNotifyEventReference& EventReference);
+
+	// FUAFAssetInstanceComponent interface
+	virtual void OnBindToInstance() override;
+
+	// Notify queue to dispatch
+	UPROPERTY()
+	FAnimNotifyQueue NotifyQueue;
+
+	TArray<const FAnimNotifyEvent*> ActiveAnimNotifyState;
+
+	UPROPERTY()
+	TArray<FAnimNotifyEventReference> ActiveAnimNotifyEventReference;
+
+	// Skeletal mesh component to 'fake' dispatch from
+	UPROPERTY()
+	TObjectPtr<USkeletalMeshComponent> SkeletalMeshComponent;
+};
+
+namespace UE::UAF
+{
+	// Event that dispatches to the module
+	struct FNotifyDispatchEvent : public FAnimNextTraitEvent
+	{
+		DECLARE_ANIM_TRAIT_EVENT(FNotifyDispatchEvent, FAnimNextTraitEvent)
+
+		// Notifies to be dispatched
+		TArray<FAnimNotifyEventReference> Notifies;
+
+		// Weight at dispatch time
+		float Weight = 1.0f;
+	};
+
+	// Event that dispatches from the module to gameplay
+	struct FNotifyQueueDispatchEvent : public FAnimNextModule_ActionEvent
+	{
+		DECLARE_ANIM_TRAIT_EVENT(FNotifyQueueDispatchEvent, FAnimNextModule_ActionEvent)
+
+		// FAnimNextModule_ActionEvent interface
+		virtual bool IsThreadSafe() const override { return bIsThreadSafe; }
+		virtual void Execute() const override;
+
+		// The various events to dispatch
+		TArray<FAnimNotifyEventReference> EventsToNotify;
+		TArray<FAnimNotifyEventReference> EventsToEnd;
+		TArray<FAnimNotifyEventReference> EventsToBegin;
+		TArray<FAnimNotifyEventReference> EventsToTick;
+
+		// Skeletal mesh component to dispatch as
+		TWeakObjectPtr<USkeletalMeshComponent> WeakSkeletalMeshComponent;
+
+		// Delta time to apply to notifies
+		float DeltaSeconds = 0.0f;
+
+		// Whether this queue thread safe to dispatch
+		bool bIsThreadSafe = false;
+	};
+}
+
+#undef UE_API
\ No newline at end of file
diff --git a/Engine/Plugins/Experimental/UAF/UAFAnimGraph/Source/UAFAnimGraph/Private/Graph/RigUnit_AnimNextSetNotifyContext.cpp b/Engine/Plugins/Experimental/UAF/UAFAnimGraph/Source/UAFAnimGraph/Private/Graph/RigUnit_AnimNextSetNotifyContext.cpp
--- a/Engine/Plugins/Experimental/UAF/UAFAnimGraph/Source/UAFAnimGraph/Private/Graph/RigUnit_AnimNextSetNotifyContext.cpp
+++ b/Engine/Plugins/Experimental/UAF/UAFAnimGraph/Source/UAFAnimGraph/Private/Graph/RigUnit_AnimNextSetNotifyContext.cpp
@@ -4,14 +4,14 @@
 
 #include "Components/SkeletalMeshComponent.h"
 #include "Module/AnimNextModuleInstance.h"
-#include "Traits/NotifyDispatcher.h"
+#include "Module/UAFNotifyDispatcherComponent.h"
 
 #include UE_INLINE_GENERATED_CPP_BY_NAME(RigUnit_AnimNextSetNotifyContext)
 
 FRigUnit_AnimNextSetNotifyContext_Execute()
 {
 	FAnimNextModuleInstance& ModuleInstance = ExecuteContext.GetContextData<FAnimNextModuleContextData>().GetModuleInstance();
-	FAnimNextNotifyDispatcherComponent& NotifyDispatcher = ModuleInstance.GetComponent<FAnimNextNotifyDispatcherComponent>();
+	FUAFNotifyDispatcherComponent& NotifyDispatcher = ModuleInstance.GetComponent<FUAFNotifyDispatcherComponent>();
 	NotifyDispatcher.SkeletalMeshComponent = SkeletalMeshComponent;
 	NotifyDispatcher.NotifyQueue.PredictedLODLevel = SkeletalMeshComponent ? SkeletalMeshComponent->GetPredictedLODLevel() : 0;
 }
diff --git a/Engine/Plugins/Experimental/UAF/UAFAnimGraph/Source/UAFAnimGraph/Private/Injection/MontageTrait.cpp b/Engine/Plugins/Experimental/UAF/UAFAnimGraph/Source/UAFAnimGraph/Private/Injection/MontageTrait.cpp
--- a/Engine/Plugins/Experimental/UAF/UAFAnimGraph/Source/UAFAnimGraph/Private/Injection/MontageTrait.cpp
+++ b/Engine/Plugins/Experimental/UAF/UAFAnimGraph/Source/UAFAnimGraph/Private/Injection/MontageTrait.cpp
@@ -13,7 +13,7 @@
 #include "EvaluationVM/Tasks/BlendKeyframes.h"
 #include "EvaluationVM/Tasks/NormalizeRotations.h"
 #include "EvaluationVM/Tasks/PushReferenceKeyframe.h"
-#include "Traits/NotifyDispatcher.h"
+#include "Module/UAFNotifyDispatcherComponent.h"
 #include "Traits/NotifyFilterTrait.h"
 #include "Injection/UAFMontageComponent.h"
 #include "Graph/SyncGroup_GraphInstanceComponent.h"
@@ -320,7 +320,7 @@ namespace UE::UAF
 				if (FAnimNextModuleInstance* ModuleInstance = Context.GetRootGraphInstance().GetModuleInstance())
 				{
 					// Ensure we have a handler component on the module
-					(void)ModuleInstance->GetComponent<FAnimNextNotifyDispatcherComponent>();
+					(void)ModuleInstance->GetComponent<FUAFNotifyDispatcherComponent>();
 				}
 
 				TSharedPtr<FNotifyDispatchEvent> NotifyDispatchEvent = MakeTraitEvent<FNotifyDispatchEvent>();
diff --git a/Engine/Plugins/Experimental/UAF/UAFAnimGraph/Source/UAFAnimGraph/Private/Injection/UAFMontageComponent.cpp b/Engine/Plugins/Experimental/UAF/UAFAnimGraph/Source/UAFAnimGraph/Private/Injection/UAFMontageComponent.cpp
--- a/Engine/Plugins/Experimental/UAF/UAFAnimGraph/Source/UAFAnimGraph/Private/Injection/UAFMontageComponent.cpp
+++ b/Engine/Plugins/Experimental/UAF/UAFAnimGraph/Source/UAFAnimGraph/Private/Injection/UAFMontageComponent.cpp
@@ -12,7 +12,7 @@
 #include "Engine/World.h"
 #include "Logging/StructuredLog.h"
 #include "Module/AnimNextModuleInstance.h"
-#include "Traits/NotifyDispatcher.h"
+#include "Module/UAFNotifyDispatcherComponent.h"
 
 
 #include UE_INLINE_GENERATED_CPP_BY_NAME(UAFMontageComponent)
diff --git a/Engine/Plugins/Experimental/UAF/UAFAnimGraph/Source/UAFAnimGraph/Private/Traits/NotifyDispatcher.cpp b/Engine/Plugins/Experimental/UAF/UAFAnimGraph/Source/UAFAnimGraph/Private/Traits/NotifyDispatcher.cpp
--- a/Engine/Plugins/Experimental/UAF/UAFAnimGraph/Source/UAFAnimGraph/Private/Traits/NotifyDispatcher.cpp
+++ b/Engine/Plugins/Experimental/UAF/UAFAnimGraph/Source/UAFAnimGraph/Private/Traits/NotifyDispatcher.cpp
@@ -3,17 +3,11 @@
 #include "Traits/NotifyDispatcher.h"
 #include "TraitInterfaces/INotifySource.h"
 #include "TraitInterfaces/ITimeline.h"
-#include "AnimationRuntime.h"
 #include "NotifyFilterTrait.h"
 #include "Animation/AnimTrace.h"
-#include "Animation/AnimNotifies/AnimNotify.h"
-#include "Animation/AnimNotifies/AnimNotifyState.h"
-#include "Components/SkeletalMeshComponent.h"
 #include "Graph/AnimNextGraphInstance.h"
 #include "Module/AnimNextModuleInstance.h"
-#include "Module/AnimNextSkeletalMeshComponentReferenceComponent.h"
-
-#include UE_INLINE_GENERATED_CPP_BY_NAME(NotifyDispatcher)
+#include "Module/UAFNotifyDispatcherComponent.h"
 
 namespace UE::UAF
 {
@@ -73,7 +67,7 @@ namespace UE::UAF
 				// Ensure we have a handler component on the module
 				if(FAnimNextModuleInstance* ModuleInstance = Context.GetRootGraphInstance().GetModuleInstance())
 				{
-					(void)ModuleInstance->GetComponent<FAnimNextNotifyDispatcherComponent>();
+					(void)ModuleInstance->GetComponent<FUAFNotifyDispatcherComponent>();
 				}
 
 				auto NotifyDispatchEvent = MakeTraitEvent<FNotifyDispatchEvent>();
@@ -83,236 +77,4 @@ namespace UE::UAF
 			}
 		}
 	}
-
-	template<typename TriggerFunc>
-	static void TriggerEventFunc(const TArray<FAnimNotifyEventReference>& InEvents, TriggerFunc InTriggerFunc)
-	{
-		for(const FAnimNotifyEventReference& EventReference : InEvents)
-		{
-			const FAnimNotifyEvent* AnimNotifyEvent = EventReference.GetNotify();
-			if(AnimNotifyEvent == nullptr)
-			{
-				continue;
-			}
-
-			InTriggerFunc(AnimNotifyEvent, EventReference);
-		}
-	}
-
-	void FNotifyQueueDispatchEvent::Execute() const
-	{
-		USkeletalMeshComponent* SkeletalMeshComponent = WeakSkeletalMeshComponent.Get();
-		if (SkeletalMeshComponent == nullptr)
-		{
-			return;
-		}
-
-		// Order of dispatch here mimics that of UAnimInstance::TriggerAnimNotifies
-
-		// Triggers...
-		TriggerEventFunc(EventsToNotify, [SkeletalMeshComponent](const FAnimNotifyEvent* InAnimNotifyEvent, const FAnimNotifyEventReference& InEventReference)
-		{
-			check(InAnimNotifyEvent->Notify);
-			InAnimNotifyEvent->Notify->Notify(SkeletalMeshComponent, Cast<UAnimSequenceBase>(InAnimNotifyEvent->Notify->GetOuter()), InEventReference);
-		});
-
-		// Ends...
-		TriggerEventFunc(EventsToEnd, [SkeletalMeshComponent](const FAnimNotifyEvent* InAnimNotifyEvent, const FAnimNotifyEventReference& InEventReference)
-		{
-			check(InAnimNotifyEvent->NotifyStateClass);
-			InAnimNotifyEvent->NotifyStateClass->NotifyEnd(SkeletalMeshComponent, Cast<UAnimSequenceBase>(InAnimNotifyEvent->NotifyStateClass->GetOuter()), InEventReference);
-		});
-
-		// Begins...
-		TriggerEventFunc(EventsToBegin, [SkeletalMeshComponent](const FAnimNotifyEvent* InAnimNotifyEvent, const FAnimNotifyEventReference& InEventReference)
-		{
-			check(InAnimNotifyEvent->NotifyStateClass);
-			InAnimNotifyEvent->NotifyStateClass->NotifyBegin(SkeletalMeshComponent, Cast<UAnimSequenceBase>(InAnimNotifyEvent->NotifyStateClass->GetOuter()), InAnimNotifyEvent->GetDuration(), InEventReference);
-		});
-
-		// Ticks...
-		TriggerEventFunc(EventsToTick, [SkeletalMeshComponent, DeltaSeconds = DeltaSeconds](const FAnimNotifyEvent* InAnimNotifyEvent, const FAnimNotifyEventReference& InEventReference)
-		{
-			check(InAnimNotifyEvent->NotifyStateClass);
-			InAnimNotifyEvent->NotifyStateClass->NotifyTick(SkeletalMeshComponent, Cast<UAnimSequenceBase>(InAnimNotifyEvent->NotifyStateClass->GetOuter()), DeltaSeconds, InEventReference);
-		});
-	}
-}
-
-void FAnimNextNotifyDispatcherComponent::OnBindToInstance()
-{
-	// Default to the current host's skeletal mesh component, if any
-	FAnimNextModuleInstance& ModuleInstance = GetModuleInstance();
-	const FAnimNextSkeletalMeshComponentReferenceComponent& ComponentReference = ModuleInstance.GetComponent<FAnimNextSkeletalMeshComponentReferenceComponent>();
-	SkeletalMeshComponent = ComponentReference.GetComponent();
-	NotifyQueue.PredictedLODLevel = SkeletalMeshComponent ? SkeletalMeshComponent->GetPredictedLODLevel() : 0;
-}
-
-void FAnimNextNotifyDispatcherComponent::OnTraitEvent(FAnimNextTraitEvent& Event)
-{
-	if(UE::UAF::FNotifyDispatchEvent* NotifyDispatchEvent = Event.AsType<UE::UAF::FNotifyDispatchEvent>())
-	{
-		NotifyQueue.AddAnimNotifies(NotifyDispatchEvent->Notifies, NotifyDispatchEvent->Weight);
-		NotifyDispatchEvent->MarkConsumed();
-	}
-}
-
-void FAnimNextNotifyDispatcherComponent::OnEndExecution(float InDeltaTime)
-{
-	using namespace UE::UAF;
-
-	// Early out if we dont need to do any work
-	if(NotifyQueue.AnimNotifies.Num() == 0 && ActiveAnimNotifyState.Num() == 0)
-	{
-		return;
-	}
-
-	FAnimNextModuleInstance& ModuleInstance = GetModuleInstance();
-
-	TSharedPtr<FNotifyQueueDispatchEvent> NotifyQueueDispatchEvent;
-	if(SkeletalMeshComponent != nullptr)
-	{
-		NotifyQueueDispatchEvent = MakeTraitEvent<FNotifyQueueDispatchEvent>();
-		NotifyQueueDispatchEvent->WeakSkeletalMeshComponent = SkeletalMeshComponent;
-		NotifyQueueDispatchEvent->DeltaSeconds = InDeltaTime;
-	}
-
-	NotifyQueue.PredictedLODLevel = SkeletalMeshComponent ? SkeletalMeshComponent->GetPredictedLODLevel() : 0;
-
-	// Logic from here largely mimics that of UAnimInstance::TriggerAnimNotifies 
-
-	// Array that will replace the 'ActiveAnimNotifyState' at the end of this function.
-	TArray<const FAnimNotifyEvent*> NewActiveAnimNotifyState;
-	NewActiveAnimNotifyState.Reserve(NotifyQueue.AnimNotifies.Num());
-
-	TArray<FAnimNotifyEventReference> NewActiveAnimNotifyEventReference;
-	NewActiveAnimNotifyEventReference.Reserve(NotifyQueue.AnimNotifies.Num());
-
-	// AnimNotifyState freshly added that need their 'NotifyBegin' event called.
-	TArray<const FAnimNotifyEvent*> NotifyStateBeginEvent;
-	TArray<const FAnimNotifyEventReference*> NotifyStateBeginEventReference;
-
-	for (int32 Index= 0 ; Index < NotifyQueue.AnimNotifies.Num(); Index++)
-	{
-		if(const FAnimNotifyEvent* AnimNotifyEvent = NotifyQueue.AnimNotifies[Index].GetNotify())
-		{
-			// AnimNotifyState
-			if (AnimNotifyEvent->NotifyStateClass)
-			{
-				int32 ExistingItemIndex = ActiveAnimNotifyState.IndexOfByPredicate([AnimNotifyEvent](const FAnimNotifyEvent* InOther)
-				{
-					return *AnimNotifyEvent == *InOther;
-				});
-				
-				if (ExistingItemIndex != INDEX_NONE)
-				{
-					check(ActiveAnimNotifyState.Num() == ActiveAnimNotifyEventReference.Num());
-					ActiveAnimNotifyState.RemoveAtSwap(ExistingItemIndex, EAllowShrinking::No);
-					ActiveAnimNotifyEventReference.RemoveAtSwap(ExistingItemIndex, EAllowShrinking::No);
-				}
-				else
-				{
-					NotifyStateBeginEvent.Add(AnimNotifyEvent);
-					NotifyStateBeginEventReference.Add(&NotifyQueue.AnimNotifies[Index]);
-				}
-
-				NewActiveAnimNotifyState.Add(AnimNotifyEvent);
-				FAnimNotifyEventReference& EventRef = NewActiveAnimNotifyEventReference.Add_GetRef(NotifyQueue.AnimNotifies[Index]);
-				EventRef.SetNotify(AnimNotifyEvent);
-				continue;
-			}
-
-			// Trigger non 'state' AnimNotifies
-			TriggerSingleAnimNotify(InDeltaTime, NotifyQueueDispatchEvent.Get(), NotifyQueue.AnimNotifies[Index]);
-		}
-	}
-
-	if (NotifyQueueDispatchEvent.IsValid())
-	{
-		// Send end notification to AnimNotifyState not active anymore.
-		for (int32 Index = 0; Index < ActiveAnimNotifyState.Num(); ++Index)
-		{
-			const FAnimNotifyEvent* AnimNotifyEvent = ActiveAnimNotifyState[Index];
-			const FAnimNotifyEventReference& EventReference = ActiveAnimNotifyEventReference[Index];
-			if (AnimNotifyEvent->NotifyStateClass)
-			{
-#if WITH_EDITOR
-				// Prevent firing notifies in animation editors if requested
-				if(ModuleInstance.GetWorldType() != EWorldType::EditorPreview || AnimNotifyEvent->NotifyStateClass->ShouldFireInEditor())
-#endif
-				{
-					NotifyQueueDispatchEvent->EventsToEnd.Add(EventReference);
-				}
-			}
-		}
-
-		check(NotifyStateBeginEventReference.Num() == NotifyStateBeginEvent.Num());
-		for (int32 Index = 0; Index < NotifyStateBeginEvent.Num(); Index++)
-		{
-			const FAnimNotifyEvent* AnimNotifyEvent = NotifyStateBeginEvent[Index];
-			const FAnimNotifyEventReference* EventReference = NotifyStateBeginEventReference[Index];
-			if(AnimNotifyEvent->NotifyStateClass)
-			{
-#if WITH_EDITOR
-				// Prevent firing notifies in animation editors if requested 
-				if(ModuleInstance.GetWorldType() != EWorldType::EditorPreview || AnimNotifyEvent->NotifyStateClass->ShouldFireInEditor())
-#endif
-				{
-					NotifyQueueDispatchEvent->EventsToBegin.Add(*EventReference);
-				}
-			}
-		}
-	}
-
-	// Switch our arrays.
-	ActiveAnimNotifyState = MoveTemp(NewActiveAnimNotifyState);
-	ActiveAnimNotifyEventReference = MoveTemp(NewActiveAnimNotifyEventReference);
-
-	if (NotifyQueueDispatchEvent.IsValid())
-	{
-		// Tick currently active AnimNotifyState
-		for (int32 Index = 0; Index < ActiveAnimNotifyState.Num(); Index++)
-		{
-			const FAnimNotifyEvent* AnimNotifyEvent = ActiveAnimNotifyState[Index];
-			const FAnimNotifyEventReference& EventReference = ActiveAnimNotifyEventReference[Index];
-			if(AnimNotifyEvent->NotifyStateClass)
-			{
-#if WITH_EDITOR
-				// Prevent firing notifies in animation editors if requested 
-				if(ModuleInstance.GetWorldType() != EWorldType::EditorPreview || AnimNotifyEvent->NotifyStateClass->ShouldFireInEditor())
-#endif
-				{
-					NotifyQueueDispatchEvent->EventsToTick.Add(EventReference);
-				}
-			}
-		}
-
-		GetModuleInstance().QueueOutputTraitEvent(NotifyQueueDispatchEvent);
-	}
-	NotifyQueue.AnimNotifies.Reset();
-}
-
-void FAnimNextNotifyDispatcherComponent::TriggerSingleAnimNotify(float InDeltaTime, UE::UAF::FNotifyQueueDispatchEvent* InDispatcher, const FAnimNotifyEventReference& EventReference)
-{
-	// This is for non 'state' anim notifies.
-	const FAnimNotifyEvent* AnimNotifyEvent = EventReference.GetNotify();
-	if (AnimNotifyEvent && (AnimNotifyEvent->NotifyStateClass == nullptr))
-	{
-		FAnimNextModuleInstance& ModuleInstance = GetModuleInstance();
-		if (InDispatcher != nullptr && AnimNotifyEvent->Notify != nullptr)
-		{
-#if WITH_EDITOR
-			// Prevent firing notifies in animation editors if requested 
-			if(ModuleInstance.GetWorldType() != EWorldType::EditorPreview || AnimNotifyEvent->Notify->ShouldFireInEditor())
-#endif
-			{
-				InDispatcher->EventsToNotify.Add(EventReference);
-			}
-		}
-		else if (AnimNotifyEvent->NotifyName != NAME_None)
-		{
-			// Named notifies can be handled by a custom module event on our worker thread, no need to queue
-			ModuleInstance.RunScriptEventByName(AnimNotifyEvent->NotifyName, InDeltaTime);
-		}
-	}
 }
diff --git a/Engine/Plugins/Experimental/UAF/UAFAnimGraph/Source/UAFAnimGraph/Private/Traits/NotifyDispatcher.h b/Engine/Plugins/Experimental/UAF/UAFAnimGraph/Source/UAFAnimGraph/Private/Traits/NotifyDispatcher.h
--- a/Engine/Plugins/Experimental/UAF/UAFAnimGraph/Source/UAFAnimGraph/Private/Traits/NotifyDispatcher.h
+++ b/Engine/Plugins/Experimental/UAF/UAFAnimGraph/Source/UAFAnimGraph/Private/Traits/NotifyDispatcher.h
@@ -3,53 +3,10 @@
 #pragma once
 
 #include "CoreMinimal.h"
-#include "Module/ModuleEvents.h"
 #include "TraitCore/Trait.h"
 #include "TraitInterfaces/ITimelinePlayer.h"
 #include "TraitInterfaces/IUpdate.h"
-#include "Animation/AnimNotifyQueue.h"
-#include "Module/UAFModuleInstanceComponent.h"
 #include "Traits/NotifyDispatcherTraitData.h"
-#include "NotifyDispatcher.generated.h"
-
-struct FAnimNotifyEventReference;
-class USkeletalMeshComponent;
-struct FAnimNextModuleInstance;
-
-namespace UE::UAF
-{
-	struct FNotifyQueueDispatchEvent;
-}
-
-USTRUCT()
-struct FAnimNextNotifyDispatcherComponent : public FUAFModuleInstanceComponent
-{
-	GENERATED_BODY()
-	DECLARE_UAF_ASSET_INSTANCE_COMPONENT()
-
-	// FUAFModuleInstanceComponent interface
-	virtual void OnTraitEvent(FAnimNextTraitEvent& Event) override;
-	virtual void OnEndExecution(float InDeltaTime) override;
-
-	// Triggers a single anim notify in the dispatcher
-	void TriggerSingleAnimNotify(float InDeltaTime, UE::UAF::FNotifyQueueDispatchEvent* InDispatcher, const FAnimNotifyEventReference& EventReference);
-
-	// FUAFAssetInstanceComponent interface
-	virtual void OnBindToInstance() override;
-
-	// Notify queue to dispatch
-	UPROPERTY()
-	FAnimNotifyQueue NotifyQueue;
-
-	TArray<const FAnimNotifyEvent*> ActiveAnimNotifyState;
-
-	UPROPERTY()
-	TArray<FAnimNotifyEventReference> ActiveAnimNotifyEventReference;
-
-	// Skeletal mesh component to 'fake' dispatch from
-	UPROPERTY()
-	TObjectPtr<USkeletalMeshComponent> SkeletalMeshComponent;
-};
 
 namespace UE::UAF
 {
@@ -71,41 +28,4 @@ namespace UE::UAF
 		// ITimelinePlayer impl
 		virtual void AdvanceBy(FExecutionContext& Context, const TTraitBinding<ITimelinePlayer>& Binding, float DeltaTime, bool bDispatchEvents) const override;
 	};
-
-	// Event that dispatches to the module
-	struct FNotifyDispatchEvent : public FAnimNextTraitEvent
-	{
-		DECLARE_ANIM_TRAIT_EVENT(FNotifyDispatchEvent, FAnimNextTraitEvent)
-
-		// Notifies to be dispatched
-		TArray<FAnimNotifyEventReference> Notifies;
-
-		// Weight at dispatch time
-		float Weight = 1.0f;
-	};
-
-	// Event that dispatches from the module to gameplay
-	struct FNotifyQueueDispatchEvent : public FAnimNextModule_ActionEvent
-	{
-		DECLARE_ANIM_TRAIT_EVENT(FNotifyQueueDispatchEvent, FAnimNextModule_ActionEvent)
-
-		// FAnimNextModule_ActionEvent interface
-		virtual bool IsThreadSafe() const override { return bIsThreadSafe; }
-		virtual void Execute() const override;
-
-		// The various events to dispatch
-		TArray<FAnimNotifyEventReference> EventsToNotify;
-		TArray<FAnimNotifyEventReference> EventsToEnd;
-		TArray<FAnimNotifyEventReference> EventsToBegin;
-		TArray<FAnimNotifyEventReference> EventsToTick;
-
-		// Skeletal mesh component to dispatch as
-		TWeakObjectPtr<USkeletalMeshComponent> WeakSkeletalMeshComponent;
-
-		// Delta time to apply to notifies
-		float DeltaSeconds = 0.0f;
-
-		// Whether this queue thread safe to dispatch
-		bool bIsThreadSafe = false;
-	};
 }
