RewindDebugger: Trace UAF evaluation program as a child of the graph's host

Allows scene graph to trace evaluation program

#Synced-CL 51720594

diff --git a/Engine/Plugins/Experimental/UAF/UAFAnimGraph/Source/UAFAnimGraph/Private/Graph/UAFAnimGraphTrace.cpp b/Engine/Plugins/Experimental/UAF/UAFAnimGraph/Source/UAFAnimGraph/Private/Graph/UAFAnimGraphTrace.cpp
--- a/Engine/Plugins/Experimental/UAF/UAFAnimGraph/Source/UAFAnimGraph/Private/Graph/UAFAnimGraphTrace.cpp
+++ b/Engine/Plugins/Experimental/UAF/UAFAnimGraph/Source/UAFAnimGraph/Private/Graph/UAFAnimGraphTrace.cpp
@@ -95,6 +95,8 @@ namespace UE::UAF
 		{
 			RecordingTime = FObjectTrace::GetWorldElapsedTime(OuterObject->GetWorld());
 		}
+		
+		uint64 HostId = RootGraph.GetHost() ? RootGraph.GetHost()->GetUniqueId() : FObjectTrace::GetObjectId(OuterObject);
 
 		TArray<uint8> ArchiveData;
 		FMemoryWriter WriterArchive(ArchiveData);
@@ -106,7 +108,7 @@ namespace UE::UAF
 		
 		UE_TRACE_LOG(UAFAnimGraph, EvaluationProgram, UAFChannel)
 				<< EvaluationProgram.Cycle(FPlatformTime::Cycles64())
-				<< EvaluationProgram.OuterObjectId(FObjectTrace::GetObjectId(OuterObject))
+				<< EvaluationProgram.OuterObjectId(HostId)
 				<< EvaluationProgram.RecordingTime(RecordingTime)
 				<< EvaluationProgram.InstanceId(InstanceId)
 				<< EvaluationProgram.ProgramData(ArchiveData.GetData(), ArchiveData.Num());
diff --git a/Engine/Plugins/Experimental/UAF/UAFAnimGraph/Source/UAFAnimGraphEditor/Private/RewindDebugger/EvaluationProgramTrack.cpp b/Engine/Plugins/Experimental/UAF/UAFAnimGraph/Source/UAFAnimGraphEditor/Private/RewindDebugger/EvaluationProgramTrack.cpp
--- a/Engine/Plugins/Experimental/UAF/UAFAnimGraph/Source/UAFAnimGraphEditor/Private/RewindDebugger/EvaluationProgramTrack.cpp
+++ b/Engine/Plugins/Experimental/UAF/UAFAnimGraph/Source/UAFAnimGraphEditor/Private/RewindDebugger/EvaluationProgramTrack.cpp
@@ -24,9 +24,9 @@ namespace UE::UAF::Editor
 
 static const FName AnimNextModulesName("AnimNextModules");
 
- FName FEvaluationProgramTrackCreator::GetTargetTypeNameInternal() const
- {
-	return UUAFComponent::StaticClass()->GetFName();
+FName FEvaluationProgramTrackCreator::GetTargetTypeNameInternal() const
+{
+ 	return FUAFAssetInstance::StaticStruct()->GetFName();
 }
 
 	
