UAF: Introduce Graph Update Loop Guard
Remove the re-entrancy check in the blend stack because it would incorrectly flag factory created chains of graphs that have been created with the same asset type (e.g. chooser) and therefore same temaplate graph, but with differnet actual objects being used in the instances of those template graphs.
Since graphs can also be parameterized, re-entrancy checks even taking the factory object into account would not be foolproof either.
Instead we add guard code to catch potential infinite loops and break out of them when detected.

#Synced-CL 51817671

diff --git a/Engine/Plugins/Experimental/UAF/UAF/Source/UAFEditor/Private/Settings/UAFEditorProjectSettings.cpp b/Engine/Plugins/Experimental/UAF/UAF/Source/UAFEditor/Private/Settings/UAFEditorProjectSettings.cpp
--- a/Engine/Plugins/Experimental/UAF/UAF/Source/UAFEditor/Private/Settings/UAFEditorProjectSettings.cpp
+++ b/Engine/Plugins/Experimental/UAF/UAF/Source/UAFEditor/Private/Settings/UAFEditorProjectSettings.cpp
@@ -8,3 +8,18 @@ UUAFEditorProjectSettings::UUAFEditorProjectSettings()
 {
 	UAFBrowserCategoryConfiguration = "/UAFAnimGraph/Internal/TABC_UAFBrowserDefaultConfig.TABC_UAFBrowserDefaultConfig";
 }
+
+void UUAFEditorProjectSettings::PostInitProperties()
+{
+	Super::PostInitProperties();
+	ImportConsoleVariableValues();
+}
+
+void UUAFEditorProjectSettings::PostEditChangeProperty(FPropertyChangedEvent& PropertyChangedEvent)
+{
+	Super::PostEditChangeProperty(PropertyChangedEvent);
+	if (PropertyChangedEvent.Property)
+	{
+		ExportValuesToConsoleVariables(PropertyChangedEvent.Property);
+	}
+}
diff --git a/Engine/Plugins/Experimental/UAF/UAF/Source/UAFEditor/Private/Settings/UAFEditorProjectSettings.h b/Engine/Plugins/Experimental/UAF/UAF/Source/UAFEditor/Private/Settings/UAFEditorProjectSettings.h
--- a/Engine/Plugins/Experimental/UAF/UAF/Source/UAFEditor/Private/Settings/UAFEditorProjectSettings.h
+++ b/Engine/Plugins/Experimental/UAF/UAF/Source/UAFEditor/Private/Settings/UAFEditorProjectSettings.h
@@ -15,9 +15,17 @@ class UUAFEditorProjectSettings : public UDeveloperSettings
 
 	UUAFEditorProjectSettings();
 
+	UAFEDITOR_API virtual void PostInitProperties() override;
+	UAFEDITOR_API virtual void PostEditChangeProperty(FPropertyChangedEvent& PropertyChangedEvent) override;
 public:
 
 	/** The configuration used the UAF browser for setting up categories */
 	UPROPERTY(config, EditAnywhere, DisplayName = "UAF Browser Configuration", Category = UAF, meta = (AllowedClasses = "/Script/UserAssetTagsEditor.TaggedAssetBrowserConfiguration"))
 	FSoftObjectPath UAFBrowserCategoryConfiguration;
+	
+	/** The maximum amount of update loops per frame before a graph is considered to be stuck in an endless loop and breaks out of the update 
+	 *  This setting is linked with the CVar UAF.MaxGraphUpdateLoopCount 
+	 */
+	UPROPERTY(EditAnywhere, Category="UAF|Graph", Config, meta=(ConsoleVariable="UAF.MaxGraphUpdateLoopCount"))
+	int32 MaxGraphUpdateLoopCount = 1000;
 };
\ No newline at end of file
diff --git a/Engine/Plugins/Experimental/UAF/UAFAnimGraph/Source/UAFAnimGraph/Private/TraitInterfaces/IUpdate.cpp b/Engine/Plugins/Experimental/UAF/UAFAnimGraph/Source/UAFAnimGraph/Private/TraitInterfaces/IUpdate.cpp
--- a/Engine/Plugins/Experimental/UAF/UAFAnimGraph/Source/UAFAnimGraph/Private/TraitInterfaces/IUpdate.cpp
+++ b/Engine/Plugins/Experimental/UAF/UAFAnimGraph/Source/UAFAnimGraph/Private/TraitInterfaces/IUpdate.cpp
@@ -10,6 +10,24 @@
 #include "Module/AnimNextModuleInstance.h"
 #include "TraitCore/TraitEventRaising.h"
 
+#ifndef DO_UAF_GRAPH_LOOPGUARD
+	#if (!(UE_BUILD_SHIPPING || UE_BUILD_TEST) || WITH_EDITOR)
+		#define DO_UAF_GRAPH_LOOPGUARD 1
+	#else
+		#define DO_UAF_GRAPH_LOOPGUARD 0
+	#endif
+#endif
+
+#if DO_UAF_GRAPH_LOOPGUARD
+static int32 GMaxUAFGraphUpdateLoopCount = 1000;
+static FAutoConsoleVariableRef CVarMaxUAFGraphUpdateLoopCount(
+	TEXT("UAF.MaxGraphUpdateLoopCount"),
+	GMaxUAFGraphUpdateLoopCount,
+	TEXT("Sets the number of update loops before graph is considered in an infinite loop and will break out of update.\n"),
+	ECVF_Default
+);
+#endif
+
 namespace UE::UAF
 {
 	namespace Private
@@ -542,6 +560,9 @@ namespace UE::UAF
 		Private::FUpdateEntry RootEntry(GraphInstance.GetGraphRootPtr(), RootState);
 		TraversalContext.PushUpdateEntry(&RootEntry);
 
+#if DO_UAF_GRAPH_LOOPGUARD
+		int32 ActiveUpdateLoopCount = 0;
+#endif		
 		while (Private::FUpdateEntry* Entry = TraversalContext.PopUpdateEntry())
 		{
 			TraversalContext.ExecutingEntry = Entry;
@@ -652,6 +673,17 @@ namespace UE::UAF
 				// We don't need this entry anymore
 				TraversalContext.PushFreeEntry(Entry);
 			}
+			
+#if DO_UAF_GRAPH_LOOPGUARD
+			// Check for loop count and break out if a potential endless loop is detected
+			++ActiveUpdateLoopCount;
+			if (!ensureMsgf(ActiveUpdateLoopCount < GMaxUAFGraphUpdateLoopCount, 
+				TEXT("Potential endless loop detected in Graph [%s] of System [%s] - will break out of the update loop"),
+				 *GraphInstance.GetAssetName().ToString(), GraphInstance.GetModuleInstance() ? *GraphInstance.GetModuleInstance()->GetAssetName().ToString() : TEXT("NONE")))
+			{
+				break;
+			}
+#endif	
 		}
 
 		// Clear our executing entry
diff --git a/Engine/Plugins/Experimental/UAF/UAFAnimGraph/Source/UAFAnimGraph/Private/Traits/BlendStackTrait.cpp b/Engine/Plugins/Experimental/UAF/UAFAnimGraph/Source/UAFAnimGraph/Private/Traits/BlendStackTrait.cpp
--- a/Engine/Plugins/Experimental/UAF/UAFAnimGraph/Source/UAFAnimGraph/Private/Traits/BlendStackTrait.cpp
+++ b/Engine/Plugins/Experimental/UAF/UAFAnimGraph/Source/UAFAnimGraph/Private/Traits/BlendStackTrait.cpp
@@ -386,20 +386,6 @@ namespace UE::UAF
 				{
 					return INDEX_NONE;
 				}
-
-				// Check for re-entrancy and early-out if we are linking back to the current instance or one of its parents
-				const FName EntryPoint = GraphRequest.AnimationGraph->DefaultEntryPoint;
-				const FAnimNextGraphInstance* OwnerGraphInstance = &Binding.GetTraitPtr().GetNodeInstance()->GetOwner();
-				while (OwnerGraphInstance != nullptr)
-				{
-					if (OwnerGraphInstance->UsesAnimationGraph(GraphRequest.AnimationGraph) && OwnerGraphInstance->UsesEntryPoint(EntryPoint))
-					{
-						UE_LOGFMT(LogAnimation, Warning, "Ignoring PushGraph request for {0}, re-entrancy detected", GraphRequest.AnimationGraph->GetFName());
-						return INDEX_NONE;
-					}
-
-					OwnerGraphInstance = OwnerGraphInstance->GetParentGraphInstance();
-				}
 				break;
 			}
 		case EGraphRequestType::Child:
