UAF: Inline subgraph trait

Adds a trait that allows subgraphs to be run in-line, with child nodes in the owning graph, in a similar vein to linked instances in anim BP.
The child node instances are seperate for each re-entered sub-graph (e.g. if the target graph pin changes, then the new instance will have newly-generated child instances rather than routing to the existing outer graph children).
Additional fixes to make this work:
- Force FAnimNextVariableReference pins to always have no sub-pins via the graph schema
- Added child handle support in nested structs & arrays to the anim graph compiler
- Fixed graph instance Thaw() to correctly re-create compiled-in components (without this, the RigVM component is removed each live-recompile - eek!)
- RigVM: Added metadata fallback to pin-struct for pins that are held on traits (previously this did not take sub-structs into account for trait pins)

#Synced-CL 51720594

diff --git a/Engine/Plugins/Experimental/UAF/UAFAnimGraph/Source/UAFAnimGraph/Private/AnimNextAnimGraphModule.cpp b/Engine/Plugins/Experimental/UAF/UAFAnimGraph/Source/UAFAnimGraph/Private/AnimNextAnimGraphModule.cpp
--- a/Engine/Plugins/Experimental/UAF/UAFAnimGraph/Source/UAFAnimGraph/Private/AnimNextAnimGraphModule.cpp
+++ b/Engine/Plugins/Experimental/UAF/UAFAnimGraph/Source/UAFAnimGraph/Private/AnimNextAnimGraphModule.cpp
@@ -34,6 +34,7 @@
 #include "Traits/BlendStackTrait.h"
 #include "Traits/CallFunction.h"
 #include "Graph/AnimNextAnimGraph.h"
+#include "Traits/InlineSubGraphTraitData.h"
 #include "UAF/UAFAssetFactory.h"
 
 
@@ -81,6 +82,7 @@ void FAnimNextAnimGraphModule::StartupModule()
 		FModifyCurveParameters::StaticStruct(),
 		FAnimNextBoneMapping::StaticStruct(),
 		FUAFCallFunctionInfo::StaticStruct(),
+		FUAFInlineSubGraphInputBinding::StaticStruct()
 	};
 
 	RigVMRegistry.RegisterStructTypes(AllowedStructTypes);
diff --git a/Engine/Plugins/Experimental/UAF/UAFAnimGraph/Source/UAFAnimGraph/Private/Graph/AnimNextGraphInstance.cpp b/Engine/Plugins/Experimental/UAF/UAFAnimGraph/Source/UAFAnimGraph/Private/Graph/AnimNextGraphInstance.cpp
--- a/Engine/Plugins/Experimental/UAF/UAFAnimGraph/Source/UAFAnimGraph/Private/Graph/AnimNextGraphInstance.cpp
+++ b/Engine/Plugins/Experimental/UAF/UAFAnimGraph/Source/UAFAnimGraph/Private/Graph/AnimNextGraphInstance.cpp
@@ -86,6 +86,15 @@ FAnimNextModuleInstance* FAnimNextGraphInstance::GetModuleInstance() const
 	return ModuleInstance.Pin().Get();
 }
 
+FUAFWeakSystemReference FAnimNextGraphInstance::GetModuleInstanceReference() const
+{
+	if(TSharedPtr<FAnimNextModuleInstance> PinnedModuleInstance = ModuleInstance.Pin())
+	{
+		return FUAFWeakSystemReference(PinnedModuleInstance->GetReference());
+	}
+	return FUAFWeakSystemReference();
+}
+
 FAnimNextGraphInstance* FAnimNextGraphInstance::GetParentGraphInstance() const
 {
 	if(IsRoot())
@@ -293,6 +302,9 @@ void FAnimNextGraphInstance::Thaw()
 {
 	if (const UUAFAnimGraph* AnimationGraph = GetAnimationGraph())
 	{
+		// Copy any built-in components from the asset
+		CopyDefaultComponents();
+
 		if (Variables.bHasBeenInitialized)
 		{
 			MigrateVariables();
@@ -302,6 +314,9 @@ void FAnimNextGraphInstance::Thaw()
 			InitializeVariables();
 		}
 
+		// Bind components now variables are set up
+		BindDefaultComponents();
+
 		{
 			UE::UAF::FExecutionContext Context(*this);
 			if(const FAnimNextTraitHandle* FoundHandle = AnimationGraph->ResolvedRootTraitHandles.Find(EntryPoint))
diff --git a/Engine/Plugins/Experimental/UAF/UAFAnimGraph/Source/UAFAnimGraph/Private/Traits/InlineSubGraph.cpp b/Engine/Plugins/Experimental/UAF/UAFAnimGraph/Source/UAFAnimGraph/Private/Traits/InlineSubGraph.cpp
new file mode 100644
--- /dev/null
+++ b/Engine/Plugins/Experimental/UAF/UAFAnimGraph/Source/UAFAnimGraph/Private/Traits/InlineSubGraph.cpp
@@ -0,0 +1,300 @@
+// Copyright Epic Games, Inc. All Rights Reserved.
+
+#include "Traits/InlineSubGraph.h"
+
+#include "TraitCore/ExecutionContext.h"
+#include "TraitCore/NodeInstance.h"
+#include "Graph/AnimNextGraphInstance.h"
+#include "Graph/VirtualValueBundle_InlineGraph.h"
+#include "Module/AnimNextModuleInstance.h"
+#include "UAF/ValueRuntime/ValueBundle.h"
+
+namespace UE::UAF
+{
+	AUTO_REGISTER_ANIM_TRAIT(FUAFInlineSubGraphTrait)
+
+	// Trait implementation boilerplate
+#define TRAIT_INTERFACE_ENUMERATOR(GeneratorMacro) \
+	GeneratorMacro(IDiscreteBlend) \
+	GeneratorMacro(IGarbageCollection) \
+	GeneratorMacro(IHierarchy) \
+	GeneratorMacro(IUpdate) \
+	GeneratorMacro(IUpdateTraversal) \
+
+	GENERATE_ANIM_TRAIT_IMPLEMENTATION(FUAFInlineSubGraphTrait, TRAIT_INTERFACE_ENUMERATOR, NULL_ANIM_TRAIT_INTERFACE_ENUMERATOR, NULL_ANIM_TRAIT_EVENT_ENUMERATOR)
+#undef TRAIT_INTERFACE_ENUMERATOR
+
+	void FUAFInlineSubGraphTrait::FInstanceData::Construct(const FExecutionContext& Context, const FTraitBinding& Binding)
+	{
+		FTrait::FInstanceData::Construct(Context, Binding);
+
+		const FSharedData* SharedData = Binding.GetSharedData<FSharedData>();
+
+		// Allocate the reference pose fallback child
+		ReferencePoseChildPtr = Context.AllocateNodeInstance(Binding, SharedData->ReferencePoseChild);
+	}
+
+	void FUAFInlineSubGraphTrait::OnBecomeRelevant(FUpdateTraversalContext& Context, const TTraitBinding<IUpdate>& Binding, const FTraitUpdateState& TraitState) const
+	{
+		const FSharedData* SharedData = Binding.GetSharedData<FSharedData>();
+		FInstanceData* InstanceData = Binding.GetInstanceData<FInstanceData>();
+
+		// Resolve the latent Inputs pin here, during update traversal, and cache the result.
+		InstanceData->ResolvedInputs = SharedData->GetInputs(Binding);
+	}
+
+	void FUAFInlineSubGraphTrait::PreUpdate(FUpdateTraversalContext& Context, const TTraitBinding<IUpdate>& Binding, const FTraitUpdateState& TraitState) const
+	{
+		const FSharedData* SharedData = Binding.GetSharedData<FSharedData>();
+		FInstanceData* InstanceData = Binding.GetInstanceData<FInstanceData>();
+
+		const bool bHasActiveSubGraph = (InstanceData->CurrentlyActiveSubGraphIndex != INDEX_NONE) && ensure(InstanceData->SubGraphSlots.IsValidIndex(InstanceData->CurrentlyActiveSubGraphIndex));
+
+		TObjectPtr<const UUAFAnimGraph> CurrentActiveAnimationGraph;
+		if (bHasActiveSubGraph)
+		{
+			FSubGraphSlot& ActiveSlot = InstanceData->SubGraphSlots[InstanceData->CurrentlyActiveSubGraphIndex];
+			CurrentActiveAnimationGraph = ActiveSlot.AnimationGraph;
+			ActiveSlot.bWasRelevant = true;
+		}
+
+		const TObjectPtr<const UUAFAnimGraph> DesiredAnimationGraph = SharedData->GetGraph(Binding);
+
+		// Note: Not guarding against re-entrant sub-graph references here as higher-level recursion checks should catch this
+
+		// Transition to a new sub-graph instance if the desired graph has changed
+		if (!bHasActiveSubGraph || CurrentActiveAnimationGraph != DesiredAnimationGraph)
+		{
+			int32 FreeSlotIndex = INDEX_NONE;
+			const int32 NumSlots = InstanceData->SubGraphSlots.Num();
+			for (int32 SlotIndex = 0; SlotIndex < NumSlots; ++SlotIndex)
+			{
+				if (InstanceData->SubGraphSlots[SlotIndex].State == ESlotState::Inactive)
+				{
+					FreeSlotIndex = SlotIndex;
+					break;
+				}
+			}
+
+			if (FreeSlotIndex == INDEX_NONE)
+			{
+				FreeSlotIndex = InstanceData->SubGraphSlots.AddDefaulted();
+			}
+
+			FSubGraphSlot& NewSlot = InstanceData->SubGraphSlots[FreeSlotIndex];
+			NewSlot.AnimationGraph = DesiredAnimationGraph;
+			NewSlot.State = DesiredAnimationGraph ? ESlotState::ActiveWithGraph : ESlotState::ActiveWithReferencePose;
+
+			const int32 OldChildIndex = InstanceData->CurrentlyActiveSubGraphIndex;
+			InstanceData->CurrentlyActiveSubGraphIndex = FreeSlotIndex;
+
+			TTraitBinding<IDiscreteBlend> DiscreteBlendTrait;
+			Binding.GetStackInterface(DiscreteBlendTrait);
+			DiscreteBlendTrait.OnBlendTransition(Context, OldChildIndex, FreeSlotIndex);
+		}
+	}
+
+	void FUAFInlineSubGraphTrait::PushInputsToInnerGraph(const FExecutionContext& Context, const TTraitBinding<IDiscreteBlend>& Binding, const FSubGraphSlot& InSlot) const
+	{
+		const FInstanceData* InstanceData = Binding.GetInstanceData<FInstanceData>();
+
+		const TArray<FUAFInlineSubGraphInputBinding>& Inputs = InstanceData->ResolvedInputs;
+
+		// Iterate in lock-step with the slot's InputChildPtrs, allocated in OnBlendInitiated
+		const int32 NumInputs = FMath::Min(Inputs.Num(), InSlot.InputChildPtrs.Num());
+		for (int32 InputIndex = 0; InputIndex < NumInputs; ++InputIndex)
+		{
+			const FTraitPtr& ChildPtr = InSlot.InputChildPtrs[InputIndex];
+			const FAnimNextVariableReference& InnerVar = Inputs[InputIndex].Variable;
+
+			if (!ChildPtr.IsValid() || InnerVar.IsNone() || !InSlot.GraphInstance.IsValid())
+			{
+				continue;
+			}
+
+			// Wrap the outer child as an inline graph bundle and write it directly into the inner graph variable
+			InSlot.GraphInstance->AccessVariable<FUAFValueBundle>(InnerVar, [&ChildPtr](FUAFValueBundle& InBundle)
+			{
+				InBundle.SetAs<FVirtualValueBundle_InlineGraph>(ChildPtr);
+			});
+		}
+	}
+
+	void FUAFInlineSubGraphTrait::QueueChildrenForTraversal(FUpdateTraversalContext& Context, const TTraitBinding<IUpdateTraversal>& Binding, const FTraitUpdateState& TraitState, FUpdateTraversalQueue& TraversalQueue) const
+	{
+		const FInstanceData* InstanceData = Binding.GetInstanceData<FInstanceData>();
+
+		// Traverse active sub-graph slots
+		const int32 NumSubGraphs = InstanceData->SubGraphSlots.Num();
+		if (NumSubGraphs > 0)
+		{
+			TTraitBinding<IDiscreteBlend> DiscreteBlendTrait;
+			Binding.GetStackInterface(DiscreteBlendTrait);
+
+			for (int32 SlotIndex = 0; SlotIndex < NumSubGraphs; ++SlotIndex)
+			{
+				const FSubGraphSlot& Slot = InstanceData->SubGraphSlots[SlotIndex];
+				if (Slot.State == ESlotState::ActiveWithGraph)
+				{
+					if (Slot.AnimationGraph == nullptr || !Slot.GraphInstance.IsValid())
+					{
+						continue;
+					}
+
+					const float BlendWeight = DiscreteBlendTrait.GetBlendWeight(Context, SlotIndex);
+					const bool bGraphHasNeverUpdated = !Slot.GraphInstance->HasUpdated();
+
+					FTraitUpdateState SlotTraitState = TraitState
+						.WithWeight(BlendWeight)
+						.AsBlendingOut(SlotIndex != InstanceData->CurrentlyActiveSubGraphIndex)
+						.AsNewlyRelevant(!Slot.bWasRelevant || bGraphHasNeverUpdated);
+
+					Slot.GraphInstance->MarkAsUpdated();
+					TraversalQueue.Push(Slot.GraphInstance->GetGraphRootPtr().IsValid() ? Slot.GraphInstance->GetGraphRootPtr() : InstanceData->ReferencePoseChildPtr, SlotTraitState);
+				}
+				else if (Slot.State == ESlotState::ActiveWithReferencePose)
+				{
+					const float BlendWeight = DiscreteBlendTrait.GetBlendWeight(Context, SlotIndex);
+					
+					FTraitUpdateState SlotTraitState = TraitState
+						.WithWeight(BlendWeight)
+						.AsBlendingOut(SlotIndex != InstanceData->CurrentlyActiveSubGraphIndex)
+						.AsNewlyRelevant(!Slot.bWasRelevant);
+					
+					TraversalQueue.Push(InstanceData->ReferencePoseChildPtr, SlotTraitState);
+				}
+			}
+		}
+	}
+
+	uint32 FUAFInlineSubGraphTrait::GetNumChildren(const FExecutionContext& Context, const TTraitBinding<IHierarchy>& Binding) const
+	{
+		const FInstanceData* InstanceData = Binding.GetInstanceData<FInstanceData>();
+
+		uint32 NumChildren = 0;
+
+		for (const FSubGraphSlot& Slot : InstanceData->SubGraphSlots)
+		{
+			if (Slot.State != ESlotState::Inactive)
+			{
+				++NumChildren;
+			}
+		}
+
+		return NumChildren;
+	}
+
+	void FUAFInlineSubGraphTrait::GetChildren(const FExecutionContext& Context, const TTraitBinding<IHierarchy>& Binding, FChildrenArray& Children) const
+	{
+		const FInstanceData* InstanceData = Binding.GetInstanceData<FInstanceData>();
+
+		for (const FSubGraphSlot& Slot : InstanceData->SubGraphSlots)
+		{
+			if (Slot.State == ESlotState::ActiveWithReferencePose)
+			{
+				Children.Add(InstanceData->ReferencePoseChildPtr);
+			}
+			else if (Slot.State == ESlotState::ActiveWithGraph)
+			{
+				Children.Add((Slot.GraphInstance.IsValid() && Slot.GraphInstance->GetGraphRootPtr().IsValid()) ? Slot.GraphInstance->GetGraphRootPtr() : InstanceData->ReferencePoseChildPtr);
+			}
+		}
+	}
+
+	float FUAFInlineSubGraphTrait::GetBlendWeight(FExecutionContext& Context, const TTraitBinding<IDiscreteBlend>& Binding, int32 ChildIndex) const
+	{
+		const FInstanceData* InstanceData = Binding.GetInstanceData<FInstanceData>();
+
+		if (ChildIndex == InstanceData->CurrentlyActiveSubGraphIndex)
+		{
+			return 1.0f;
+		}
+		else if (InstanceData->SubGraphSlots.IsValidIndex(ChildIndex))
+		{
+			return 0.0f;
+		}
+
+		return -1.0f;
+	}
+
+	int32 FUAFInlineSubGraphTrait::GetBlendDestinationChildIndex(FExecutionContext& Context, const TTraitBinding<IDiscreteBlend>& Binding) const
+	{
+		const FInstanceData* InstanceData = Binding.GetInstanceData<FInstanceData>();
+		return InstanceData->CurrentlyActiveSubGraphIndex;
+	}
+
+	void FUAFInlineSubGraphTrait::OnBlendTransition(FExecutionContext& Context, const TTraitBinding<IDiscreteBlend>& Binding, int32 OldChildIndex, int32 NewChildIndex) const
+	{
+		TTraitBinding<IDiscreteBlend> DiscreteBlendTrait;
+		Binding.GetStackInterface(DiscreteBlendTrait);
+
+		// Immediate initiate-then-terminate keeps behaviour consistent with the base SubGraph trait
+		DiscreteBlendTrait.OnBlendInitiated(Context, NewChildIndex);
+		DiscreteBlendTrait.OnBlendTerminated(Context, OldChildIndex);
+	}
+
+	void FUAFInlineSubGraphTrait::OnBlendInitiated(FExecutionContext& Context, const TTraitBinding<IDiscreteBlend>& Binding, int32 ChildIndex) const
+	{
+		FInstanceData* InstanceData = Binding.GetInstanceData<FInstanceData>();
+
+		if (InstanceData->SubGraphSlots.IsValidIndex(ChildIndex))
+		{
+			FSubGraphSlot& Slot = InstanceData->SubGraphSlots[ChildIndex];
+			if (Slot.State == ESlotState::ActiveWithGraph)
+			{
+				const TArray<FUAFInlineSubGraphInputBinding>& Inputs = InstanceData->ResolvedInputs;
+
+				// Allocate a separate set of outer child ptrs for this slot so blending-in and
+				// blending-out slots each independently re-enter the outer graph traversal
+				Slot.InputChildPtrs.SetNum(Inputs.Num());
+				for (int32 i = 0; i < Inputs.Num(); ++i)
+				{
+					Slot.InputChildPtrs[i] = Context.AllocateNodeInstance(Binding, Inputs[i].Input);
+				}
+
+				FAnimNextGraphInstance& Owner = Binding.GetTraitPtr().GetNodeInstance()->GetOwner();
+				Slot.GraphInstance = Slot.AnimationGraph->AllocateInstance(
+					{
+						.SystemReference = Owner.GetModuleInstanceReference(),
+						.ParentGraphInstance = Owner.AsShared(),
+					});
+
+				PushInputsToInnerGraph(Context, Binding, Slot);
+			}
+		}
+	}
+
+	void FUAFInlineSubGraphTrait::OnBlendTerminated(FExecutionContext& Context, const TTraitBinding<IDiscreteBlend>& Binding, int32 ChildIndex) const
+	{
+		FInstanceData* InstanceData = Binding.GetInstanceData<FInstanceData>();
+
+		if (InstanceData->SubGraphSlots.IsValidIndex(ChildIndex))
+		{
+			FSubGraphSlot& Slot = InstanceData->SubGraphSlots[ChildIndex];
+			if (Slot.State == ESlotState::ActiveWithGraph)
+			{
+				Slot.GraphInstance.Reset();
+			}
+			Slot.InputChildPtrs.Reset();
+			Slot.State = ESlotState::Inactive;
+			Slot.bWasRelevant = false;
+		}
+	}
+
+	void FUAFInlineSubGraphTrait::AddReferencedObjects(const FExecutionContext& Context, const TTraitBinding<IGarbageCollection>& Binding, FReferenceCollector& Collector) const
+	{
+		IGarbageCollection::AddReferencedObjects(Context, Binding, Collector);
+
+		FInstanceData* InstanceData = Binding.GetInstanceData<FInstanceData>();
+
+		for (FSubGraphSlot& Slot : InstanceData->SubGraphSlots)
+		{
+			Collector.AddReferencedObject(Slot.AnimationGraph);
+			if (FAnimNextGraphInstance* ImplPtr = Slot.GraphInstance.Get())
+			{
+				Collector.AddPropertyReferencesWithStructARO(FAnimNextGraphInstance::StaticStruct(), ImplPtr);
+			}
+		}
+	}
+
+} // namespace UE::UAF
diff --git a/Engine/Plugins/Experimental/UAF/UAFAnimGraph/Source/UAFAnimGraph/Private/Traits/InlineSubGraph.h b/Engine/Plugins/Experimental/UAF/UAFAnimGraph/Source/UAFAnimGraph/Private/Traits/InlineSubGraph.h
new file mode 100644
--- /dev/null
+++ b/Engine/Plugins/Experimental/UAF/UAFAnimGraph/Source/UAFAnimGraph/Private/Traits/InlineSubGraph.h
@@ -0,0 +1,113 @@
+// Copyright Epic Games, Inc. All Rights Reserved.
+
+#pragma once
+
+#include "CoreMinimal.h"
+
+#include "TraitCore/Trait.h"
+#include "TraitInterfaces/IDiscreteBlend.h"
+#include "TraitInterfaces/IGarbageCollection.h"
+#include "TraitInterfaces/IHierarchy.h"
+#include "TraitInterfaces/IUpdate.h"
+#include "Graph/AnimNextAnimationGraph.h"
+#include "Graph/AnimNextGraphInstance.h"
+#include "Traits/InlineSubGraphTraitData.h"
+
+namespace UE::UAF
+{
+
+/**
+ * FUAFInlineSubGraphTrait
+ *
+ * A trait that hosts and manages a sub-graph instance while maintaining a set of inline child
+ * handles that act as passthrough connections to the outer owning graph.
+ * Each blending slot owns its own set of allocated input child ptrs so that blending-in and
+ * blending-out instances remain independent.
+ */
+struct FUAFInlineSubGraphTrait
+	: FBaseTrait
+	, IUpdate
+	, IUpdateTraversal
+	, IHierarchy
+	, IDiscreteBlend
+	, IGarbageCollection
+{
+	DECLARE_ANIM_TRAIT(FUAFInlineSubGraphTrait, FBaseTrait)
+
+	enum class ESlotState : uint8
+	{
+		ActiveWithGraph,
+		ActiveWithReferencePose,
+		Inactive,
+	};
+
+	struct FSubGraphSlot
+	{
+		// The animation graph asset for this slot
+		TObjectPtr<const UUAFAnimGraph> AnimationGraph;
+
+		// The live graph instance for this slot
+		TSharedPtr<FAnimNextGraphInstance> GraphInstance;
+
+		// FTraitPtrs allocated from ResolvedInputs during OnBlendInitiated, one per Inputs entry.
+		// Each ptr is wrapped as a FVirtualValueBundle_InlineGraph and written into the inner
+		// graph's corresponding interface variable. Allocated per slot so blending-in and
+		// blending-out instances each hold an independent set.
+		TArray<FTraitPtr> InputChildPtrs;
+
+		// Current lifecycle state of the slot
+		ESlotState State = ESlotState::Inactive;
+
+		// Whether this slot was relevant on the previous update tick
+		bool bWasRelevant = false;
+	};
+
+	using FSharedData = FUAFInlineSubGraphTraitSharedData;
+
+	struct FInstanceData : FTrait::FInstanceData
+	{
+		void Construct(const FExecutionContext& Context, const FTraitBinding& Binding);
+
+		// Sub-graph slots managing hosted inner graph instances
+		TArray<FSubGraphSlot> SubGraphSlots;
+
+		// Index of the currently active sub-graph slot (INDEX_NONE if none)
+		int32 CurrentlyActiveSubGraphIndex = INDEX_NONE;
+
+		// Child used to output the reference pose when no graph is provided
+		FTraitPtr ReferencePoseChildPtr;
+
+		// Inputs resolved from the latent pin during OnBecomeRelevant, cached here so that
+		// OnBlendInitiated (which runs in an execution context) can use them without re-accessing
+		// the latent system outside of update traversal.
+		TArray<FUAFInlineSubGraphInputBinding> ResolvedInputs;
+	};
+
+	// IUpdate impl
+	virtual void OnBecomeRelevant(FUpdateTraversalContext& Context, const TTraitBinding<IUpdate>& Binding, const FTraitUpdateState& TraitState) const override;
+	virtual void PreUpdate(FUpdateTraversalContext& Context, const TTraitBinding<IUpdate>& Binding, const FTraitUpdateState& TraitState) const override;
+
+	// IUpdateTraversal impl
+	virtual void QueueChildrenForTraversal(FUpdateTraversalContext& Context, const TTraitBinding<IUpdateTraversal>& Binding, const FTraitUpdateState& TraitState, FUpdateTraversalQueue& TraversalQueue) const override;
+
+	// IHierarchy impl
+	virtual uint32 GetNumChildren(const FExecutionContext& Context, const TTraitBinding<IHierarchy>& Binding) const override;
+	virtual void GetChildren(const FExecutionContext& Context, const TTraitBinding<IHierarchy>& Binding, FChildrenArray& Children) const override;
+
+	// IDiscreteBlend impl
+	virtual float GetBlendWeight(FExecutionContext& Context, const TTraitBinding<IDiscreteBlend>& Binding, int32 ChildIndex) const override;
+	virtual int32 GetBlendDestinationChildIndex(FExecutionContext& Context, const TTraitBinding<IDiscreteBlend>& Binding) const override;
+	virtual void OnBlendTransition(FExecutionContext& Context, const TTraitBinding<IDiscreteBlend>& Binding, int32 OldChildIndex, int32 NewChildIndex) const override;
+	virtual void OnBlendInitiated(FExecutionContext& Context, const TTraitBinding<IDiscreteBlend>& Binding, int32 ChildIndex) const override;
+	virtual void OnBlendTerminated(FExecutionContext& Context, const TTraitBinding<IDiscreteBlend>& Binding, int32 ChildIndex) const override;
+
+	// IGarbageCollection impl
+	virtual void AddReferencedObjects(const FExecutionContext& Context, const TTraitBinding<IGarbageCollection>& Binding, FReferenceCollector& Collector) const override;
+
+private:
+	// For each entry in ResolvedInputs, wraps the slot's corresponding FTraitPtr as a
+	// FVirtualValueBundle_InlineGraph and writes it into the inner graph instance via AccessVariable.
+	void PushInputsToInnerGraph(const FExecutionContext& Context, const TTraitBinding<IDiscreteBlend>& Binding, const FSubGraphSlot& InSlot) const;
+};
+
+} // namespace UE::UAF
diff --git a/Engine/Plugins/Experimental/UAF/UAFAnimGraph/Source/UAFAnimGraph/Private/Traits/InputValueTrait.cpp b/Engine/Plugins/Experimental/UAF/UAFAnimGraph/Source/UAFAnimGraph/Private/Traits/InputValueTrait.cpp
--- a/Engine/Plugins/Experimental/UAF/UAFAnimGraph/Source/UAFAnimGraph/Private/Traits/InputValueTrait.cpp
+++ b/Engine/Plugins/Experimental/UAF/UAFAnimGraph/Source/UAFAnimGraph/Private/Traits/InputValueTrait.cpp
@@ -98,14 +98,17 @@ namespace UE::UAF
 					TSharedPtr<FAnimNextGraphInstance> GraphInstance = GraphReference->GetOrAllocate(Context, Binding);
 					if (GraphInstance.IsValid() && GraphInstance->GetGraphRootPtr().IsValid())
 					{
-						TraversalQueue.Push(GraphInstance->GetGraphRootPtr(), TraitState);
+						TraversalQueue.Push(GraphInstance->GetGraphRootPtr(), TraitState.AsNewlyRelevant(TraitState.IsNewlyRelevant() || !GraphInstance->HasUpdated()));
+						GraphInstance->MarkAsUpdated();
 					}
 				}
 				else if (const FTraitPtr* TraitPtr = Impl->GetInlineGraph())
 				{
 					if (TraitPtr->IsValid())
 					{
-						TraversalQueue.Push(*TraitPtr, TraitState);
+						FAnimNextGraphInstance& Owner = TraitPtr->GetNodeInstance()->GetOwner();
+						TraversalQueue.Push(*TraitPtr, TraitState.AsNewlyRelevant(TraitState.IsNewlyRelevant() || !Owner.HasUpdated()));
+						Owner.MarkAsUpdated();
 					}
 				}
 			}
diff --git a/Engine/Plugins/Experimental/UAF/UAFAnimGraph/Source/UAFAnimGraph/Public/Graph/AnimNextGraphInstance.h b/Engine/Plugins/Experimental/UAF/UAFAnimGraph/Source/UAFAnimGraph/Public/Graph/AnimNextGraphInstance.h
--- a/Engine/Plugins/Experimental/UAF/UAFAnimGraph/Source/UAFAnimGraph/Public/Graph/AnimNextGraphInstance.h
+++ b/Engine/Plugins/Experimental/UAF/UAFAnimGraph/Source/UAFAnimGraph/Public/Graph/AnimNextGraphInstance.h
@@ -32,6 +32,7 @@ struct FRigUnit_AnimNextRunAnimationGraph_v2;
 struct FAnimNextStateTreeRigVMConditionBase;
 struct FAnimNextStateTreeRigVMTaskBase;
 struct FAnimNextFactoryParams;
+struct FUAFWeakSystemReference;
 
 namespace UE::UAF
 {
@@ -91,6 +92,9 @@ struct FAnimNextGraphInstance : public FUAFAssetInstance
 	// Returns the module instance that owns us or nullptr if we are invalid
 	UE_API FAnimNextModuleInstance* GetModuleInstance() const;
 
+	// Returns the system instance that owns us as a weak reference. Can be invalid if we are not owned by a system.
+	UE_API FUAFWeakSystemReference GetModuleInstanceReference() const;
+
 	// Returns the parent graph instance that owns us or nullptr for the root graph instance or if we are invalid
 	UE_API FAnimNextGraphInstance* GetParentGraphInstance() const;
 
diff --git a/Engine/Plugins/Experimental/UAF/UAFAnimGraph/Source/UAFAnimGraph/Public/Traits/InlineSubGraphTraitData.h b/Engine/Plugins/Experimental/UAF/UAFAnimGraph/Source/UAFAnimGraph/Public/Traits/InlineSubGraphTraitData.h
new file mode 100644
--- /dev/null
+++ b/Engine/Plugins/Experimental/UAF/UAFAnimGraph/Source/UAFAnimGraph/Public/Traits/InlineSubGraphTraitData.h
@@ -0,0 +1,73 @@
+// Copyright Epic Games, Inc. All Rights Reserved.
+
+#pragma once
+
+#include "TraitCore/TraitBinding.h"
+#include "TraitCore/TraitHandle.h"
+#include "TraitCore/TraitSharedData.h"
+#include "Variables/AnimNextVariableReference.h"
+
+#include "InlineSubGraphTraitData.generated.h"
+
+class UUAFAnimGraph;
+
+/**
+ * A single binding that connects a child node in this graph to a Value Bundle
+ * variable on the inner sub-graph's interface.
+ */
+USTRUCT(meta=(Hidden))
+struct FUAFInlineSubGraphInputBinding
+{
+	GENERATED_BODY()
+
+	/** A child node in this graph whose output is passed into the inner sub-graph. */
+	UPROPERTY(EditAnywhere, Category = "Sub Graph")
+	FAnimNextTraitHandle Input;
+
+	/**
+	 * The Value Bundle variable on the inner sub-graph that receives the value from Input.
+	 * Must match a variable declared on the inner sub-graph asset.
+	 */
+	UPROPERTY(EditAnywhere, Category = "Sub Graph", meta = (HideSubPins, AllowedType = FUAFValueBundle))
+	FAnimNextVariableReference Variable;
+};
+
+/**
+ * Hosts an inner sub-graph and exposes a set of child nodes from this graph as typed
+ * FUAFValueBundle inputs on the inner sub-graph's interface.
+ *
+ * Each entry in Inputs maps a child node in this graph to a Value Bundle variable
+ * declared on the inner sub-graph. When the inner sub-graph becomes active those values
+ * are bound once and remain constant for the lifetime of that graph instance.
+ *
+ * The inner sub-graph can be swapped at runtime via the Graph pin - blending between instances is supported.
+ */
+USTRUCT(meta = (DisplayName = "Inline Sub Graph", ShowTooltip = true))
+struct FUAFInlineSubGraphTraitSharedData : public FAnimNextTraitSharedData
+{
+	GENERATED_BODY()
+
+	/** The inner sub-graph asset to host. */
+	UPROPERTY(EditAnywhere, Category = "Sub Graph", meta = (ExportAsReference = "true"))
+	TObjectPtr<const UUAFAnimGraph> Graph;
+
+	/** Hidden reference pose fallback child, auto-populated during graph compilation. */
+	UPROPERTY(meta = (Hidden))
+	FAnimNextTraitHandle ReferencePoseChild;
+
+	/**
+	 * Connections from child nodes in this graph to Value Bundle variables on the inner
+	 * sub-graph's interface. The inner sub-graph must declare each referenced variable with
+	 * type FUAFValueBundle.
+	 */
+	UPROPERTY(EditAnywhere, Category = "Sub Graph", meta = (OnBecomeRelevant))
+	TArray<FUAFInlineSubGraphInputBinding> Inputs;
+
+	// Latent pin support boilerplate
+#define TRAIT_LATENT_PROPERTIES_ENUMERATOR(GeneratorMacro) \
+	GeneratorMacro(Graph) \
+	GeneratorMacro(Inputs) \
+
+	GENERATE_TRAIT_LATENT_PROPERTIES(FUAFInlineSubGraphTraitSharedData, TRAIT_LATENT_PROPERTIES_ENUMERATOR)
+#undef TRAIT_LATENT_PROPERTIES_ENUMERATOR
+};
diff --git a/Engine/Plugins/Experimental/UAF/UAFAnimGraph/Source/UAFAnimGraphUncookedOnly/Private/Graph/AnimNextAnimationGraphSchema.cpp b/Engine/Plugins/Experimental/UAF/UAFAnimGraph/Source/UAFAnimGraphUncookedOnly/Private/Graph/AnimNextAnimationGraphSchema.cpp
--- a/Engine/Plugins/Experimental/UAF/UAFAnimGraph/Source/UAFAnimGraphUncookedOnly/Private/Graph/AnimNextAnimationGraphSchema.cpp
+++ b/Engine/Plugins/Experimental/UAF/UAFAnimGraph/Source/UAFAnimGraphUncookedOnly/Private/Graph/AnimNextAnimationGraphSchema.cpp
@@ -3,6 +3,7 @@
 #include "Graph/AnimNextAnimationGraphSchema.h"
 #include "AnimGraphUncookedOnlyUtils.h"
 #include "TraitCore/TraitHandle.h"
+#include "Variables/AnimNextVariableReference.h"
 
 #include UE_INLINE_GENERATED_CPP_BY_NAME(AnimNextAnimationGraphSchema)
 
@@ -61,5 +62,10 @@ bool UAnimNextAnimationGraphSchema::ShouldUnfoldStruct(URigVMController* InContr
 		return false;
 	}
 
+	if (InStruct->IsChildOf<FAnimNextVariableReference>())
+	{
+		return false;
+	}
+
 	return true;
 }
\ No newline at end of file
diff --git a/Engine/Plugins/Experimental/UAF/UAFAnimGraph/Source/UAFAnimGraphUncookedOnly/Private/Graph/AnimNextAnimationGraph_EditorData.cpp b/Engine/Plugins/Experimental/UAF/UAFAnimGraph/Source/UAFAnimGraphUncookedOnly/Private/Graph/AnimNextAnimationGraph_EditorData.cpp
--- a/Engine/Plugins/Experimental/UAF/UAFAnimGraph/Source/UAFAnimGraphUncookedOnly/Private/Graph/AnimNextAnimationGraph_EditorData.cpp
+++ b/Engine/Plugins/Experimental/UAF/UAFAnimGraph/Source/UAFAnimGraphUncookedOnly/Private/Graph/AnimNextAnimationGraph_EditorData.cpp
@@ -143,8 +143,15 @@ namespace UE::UAF::UncookedOnly::Private
 
 	FString GetTraitProperty(const FTraitStackMapping& TraitStack, uint32 TraitIndex, FName PropertyName, const TArray<FTraitStackMapping>& TraitStackNodes)
 	{
+		URigVMGraph* Graph = TraitStack.DecoratorStackNode->GetGraph();
+		check(Graph);
+		IRigVMClientHost* ClientHost = Graph->GetImplementingOuter<IRigVMClientHost>();
+		check(ClientHost);
+		URigVMController* Controller = ClientHost->GetOrCreateController(Graph);
+		check(Controller);
+
 		const TArray<URigVMPin*>& Pins = TraitStack.TraitEntries[TraitIndex].DecoratorEntryPin->GetSubPins();
-		for (const URigVMPin* Pin : Pins)
+		for (URigVMPin* Pin : Pins)
 		{
 			if (Pin->GetDirection() != ERigVMPinDirection::Input &&
 				Pin->GetDirection() != ERigVMPinDirection::Hidden)
@@ -154,10 +161,12 @@ namespace UE::UAF::UncookedOnly::Private
 
 			if (Pin->GetFName() == PropertyName)
 			{
-				if (Pin->GetCPPTypeObject() == FAnimNextTraitHandle::StaticStruct())
+				// If we contain child handles, we need to patch them to the correct source node/trait indices
+				// Note patching the pins in this way here is only OK because we operate on a temporary intermediate copy graph, not the authored graph
+				auto PatchDefaultValue = [&TraitStackNodes, Controller](URigVMPin* InPin)
 				{
 					// Trait handle pins don't have a value, just an optional link
-					const TArray<URigVMLink*>& PinLinks = Pin->GetLinks();
+					const TArray<URigVMLink*>& PinLinks = InPin->GetLinks();
 					if (!PinLinks.IsEmpty())
 					{
 						// Something is connected to us, find the corresponding node handle so that we can encode it as our property value
@@ -192,16 +201,27 @@ namespace UE::UAF::UncookedOnly::Private
 
 							FString PropertyValue;
 							Property->ExportText_Direct(PropertyValue, &TraitHandle, &DefaultTraitHandle, nullptr, PPF_SerializedAsImportText);
-
-							return PropertyValue;
+							Controller->SetPinDefaultValue(InPin, PropertyValue, false, false, false);
 						}
 					}
+				};
 
-					// This handle pin isn't connected
-					return FString();
+				if (Pin->GetCPPTypeObject() == FAnimNextTraitHandle::StaticStruct())
+				{
+					PatchDefaultValue(Pin);
+				}
+				else
+				{
+					const TArray<URigVMPin*>& SubPins = Pin->GetAllSubPinsRecursively();
+					for (URigVMPin* SubPin : SubPins)
+					{
+						if (SubPin->GetCPPTypeObject() == FAnimNextTraitHandle::StaticStruct())
+						{
+							PatchDefaultValue(SubPin);
+						}
+					}
 				}
 
-				// A regular property pin or hidden pin
 				return Pin->GetDefaultValue();
 			}
 		}
diff --git a/Engine/Plugins/Experimental/UAF/UAFAnimGraph/Source/UAFAnimGraphUncookedOnly/Private/Templates/UAFGraphNodeTemplate_InlineSubGraph.h b/Engine/Plugins/Experimental/UAF/UAFAnimGraph/Source/UAFAnimGraphUncookedOnly/Private/Templates/UAFGraphNodeTemplate_InlineSubGraph.h
new file mode 100644
--- /dev/null
+++ b/Engine/Plugins/Experimental/UAF/UAFAnimGraph/Source/UAFAnimGraphUncookedOnly/Private/Templates/UAFGraphNodeTemplate_InlineSubGraph.h
@@ -0,0 +1,111 @@
+// Copyright Epic Games, Inc. All Rights Reserved.
+
+#pragma once
+
+#include "AnimGraphUncookedOnlyUtils.h"
+#include "AnimNextController.h"
+#include "Styling/SlateIconFinder.h"
+#include "Templates/UAFGraphNodeTemplate.h"
+#include "Templates/GraphNodeColors.h"
+#include "Traits/InlineSubGraphTraitData.h"
+#include "Graph/AnimNextAnimationGraph.h"
+#include "UAFGraphNodeTemplate_InlineSubGraph.generated.h"
+
+#define LOCTEXT_NAMESPACE "UAFGraphNodeTemplate_InlineSubGraph"
+
+UCLASS()
+class UUAFGraphNodeTemplate_InlineSubGraph : public UUAFGraphNodeTemplate
+{
+	GENERATED_BODY()
+
+	UUAFGraphNodeTemplate_InlineSubGraph()
+	{
+		Title = LOCTEXT("InlineSubGraphTitle", "Inline SubGraph");
+		TooltipText = LOCTEXT("InlineSubGraphTooltip", "Hosts an inner sub-graph and wires child nodes from this graph into the inner sub-graph's inputs");
+		Category = LOCTEXT("InlineSubGraphCategory", "UAF");
+		MenuDescription = LOCTEXT("InlineSubGraphMenuDesc", "Inline SubGraph");
+		Color = UE::UAF::UncookedOnly::FGraphNodeColors::SubGraphs;
+		Icon = *FSlateIconFinder::FindIconForClass(UUAFAnimGraph::StaticClass()).GetIcon();
+		Traits =
+		{
+			TInstancedStruct<FUAFInlineSubGraphTraitSharedData>::Make()
+		};
+		DragDropAssetTypes.Add(UUAFAnimGraph::StaticClass());
+		SetCategoryForPinsInLayout(
+			{
+				GET_PIN_PATH_STRING_CHECKED(FUAFInlineSubGraphTraitSharedData, Graph),
+				GET_PIN_PATH_STRING_CHECKED(FUAFInlineSubGraphTraitSharedData, Inputs),
+			},
+			FRigVMPinCategory::GetDefaultCategoryName(),
+			NodeLayout,
+			true);
+	}
+
+	virtual void HandleAssetDropped_Implementation(UAnimNextController* Controller, URigVMUnitNode* Node, UObject* Asset) const
+	{
+		const bool bSetupUndoRedo = !GIsTransacting;
+		if (Asset)
+		{
+			if (bSetupUndoRedo)
+			{
+				Controller->OpenUndoBracket(LOCTEXT("ConfigureNodeOnDrop", "Configure Node On Drop").ToString());
+			}
+
+			Controller->SetPinDefaultValue(
+				GET_PIN_PATH_STRING_CHECKED(FUAFInlineSubGraphTraitSharedData, Graph),
+				Asset->GetPathName(), true, bSetupUndoRedo, true, true, true);
+
+			Controller->SetNodeTitle(Node, FText::Format(LOCTEXT("NodeTitleFormat", "Run {0}"), FText::FromString(Asset->GetName())).ToString(), bSetupUndoRedo, true, true);
+
+			if (bSetupUndoRedo)
+			{
+				Controller->CloseUndoBracket();
+			}
+		}
+	}
+
+	virtual void HandlePinDefaultValueChanged_Implementation(UAnimNextController* Controller, URigVMPin* Pin) const
+	{
+		Super::HandlePinDefaultValueChanged_Implementation(Controller, Pin);
+
+		const bool bSetupUndoRedo = !GIsTransacting;
+		if (Pin->GetFName() == GET_MEMBER_NAME_CHECKED(FUAFInlineSubGraphTraitSharedData, Graph))
+		{
+			if (bSetupUndoRedo)
+			{
+				Controller->OpenUndoBracket(LOCTEXT("SetNodeTitle", "Set Node Title").ToString());
+			}
+
+			FSoftObjectPath Path(Pin->GetDefaultValue());
+			Controller->SetNodeTitle(Pin->GetNode(), FText::Format(LOCTEXT("NodeTitleFormat", "Run {0}"), FText::FromString(Path.GetAssetName())).ToString(), bSetupUndoRedo, true, true);
+
+			if (bSetupUndoRedo)
+			{
+				Controller->CloseUndoBracket();
+			}
+		}
+	}
+
+	virtual void HandleAssetRenamed_Implementation(UAnimNextController* Controller, URigVMNode* Node, const FAssetData& AssetData, const FString& OldName) const
+	{
+		Super::HandleAssetRenamed_Implementation(Controller, Node, AssetData, OldName);
+
+		// If the asset is not loaded, it can't be our asset so we can safely early out
+		UObject* Object = AssetData.FastGetAsset();
+		if (Object == nullptr)
+		{
+			return;
+		}
+
+		// NOTE: Not making this undo-able as this is an automated process, not a user action on the node specifically
+
+		URigVMPin* Pin = Node->FindPin(GET_PIN_PATH_STRING_CHECKED(FUAFInlineSubGraphTraitSharedData, Graph));
+		if (Pin && Object == Pin->GetDefaultValueObject())
+		{
+			FSoftObjectPath Path(Object);
+			Controller->SetNodeTitle(Node, FText::Format(LOCTEXT("NodeTitleFormat", "Run {0}"), FText::FromString(Path.GetAssetName())).ToString(), false, true, true);
+		}
+	}
+};
+
+#undef LOCTEXT_NAMESPACE
diff --git a/Engine/Plugins/Runtime/RigVM/Source/RigVMDeveloper/Private/RigVMModel/RigVMPin.cpp b/Engine/Plugins/Runtime/RigVM/Source/RigVMDeveloper/Private/RigVMModel/RigVMPin.cpp
--- a/Engine/Plugins/Runtime/RigVM/Source/RigVMDeveloper/Private/RigVMModel/RigVMPin.cpp
+++ b/Engine/Plugins/Runtime/RigVM/Source/RigVMDeveloper/Private/RigVMModel/RigVMPin.cpp
@@ -1541,17 +1541,16 @@ FString URigVMPin::GetMetaData(FName InKey) const
 							return Property->GetMetaData(InKey);
 						}
 					}
+
 				}
 			}
 		}
-		else
+
+		if(const UScriptStruct* Struct = GetParentScriptStruct(UnitNode))
 		{
-			if(const UScriptStruct* Struct = GetParentScriptStruct(UnitNode))
+			if(const FProperty* Property = Struct->FindPropertyByName(GetFName()))
 			{
-				if(const FProperty* Property = Struct->FindPropertyByName(GetFName()))
-				{
-					return Property->GetMetaData(InKey);
-				}
+				return Property->GetMetaData(InKey);
 			}
 		}
 	}
