UAF: Remove #if directive around console variable CVarMaxUAFGraphUpdateLoopCount because it can cause crashes when DO_UAF_GRAPH_LOOPGUARD is false but the UDeveloperSettings still try to read in the default CVar value based on the UProp meta tag

#Synced-CL 51879745

diff --git a/Engine/Plugins/Experimental/UAF/UAFAnimGraph/Source/UAFAnimGraph/Private/TraitInterfaces/IUpdate.cpp b/Engine/Plugins/Experimental/UAF/UAFAnimGraph/Source/UAFAnimGraph/Private/TraitInterfaces/IUpdate.cpp
--- a/Engine/Plugins/Experimental/UAF/UAFAnimGraph/Source/UAFAnimGraph/Private/TraitInterfaces/IUpdate.cpp
+++ b/Engine/Plugins/Experimental/UAF/UAFAnimGraph/Source/UAFAnimGraph/Private/TraitInterfaces/IUpdate.cpp
@@ -18,7 +18,6 @@
 	#endif
 #endif
 
-#if DO_UAF_GRAPH_LOOPGUARD
 static int32 GMaxUAFGraphUpdateLoopCount = 1000;
 static FAutoConsoleVariableRef CVarMaxUAFGraphUpdateLoopCount(
 	TEXT("UAF.MaxGraphUpdateLoopCount"),
@@ -26,7 +25,6 @@ static FAutoConsoleVariableRef CVarMaxUAFGraphUpdateLoopCount(
 	TEXT("Sets the number of update loops before graph is considered in an infinite loop and will break out of update.\n"),
 	ECVF_Default
 );
-#endif
 
 namespace UE::UAF
 {
