UAF: If the blend stack trait could not return a valid GraphInstance set the passed in GraphInstance ptr to nullptr

#Synced-CL 52000284

diff --git a/Engine/Plugins/Experimental/UAF/UAFAnimGraph/Source/UAFAnimGraph/Private/Traits/BlendStackTrait.cpp b/Engine/Plugins/Experimental/UAF/UAFAnimGraph/Source/UAFAnimGraph/Private/Traits/BlendStackTrait.cpp
--- a/Engine/Plugins/Experimental/UAF/UAFAnimGraph/Source/UAFAnimGraph/Private/Traits/BlendStackTrait.cpp
+++ b/Engine/Plugins/Experimental/UAF/UAFAnimGraph/Source/UAFAnimGraph/Private/Traits/BlendStackTrait.cpp
@@ -439,7 +439,7 @@ namespace UE::UAF
 			OutGraphInstance = GraphState.Instance.Get();
 			return CurrentlyActiveGraphIndex;
 		}
-
+		OutGraphInstance = nullptr;
 		return INDEX_NONE;
 	}
 
