UAF: Fix incorrect reinitialization of external graphs due to relevancy

Fixes pop in scene graph implementation when animations at lower stack levels restart.
External graphs should only be tracking their own relevancy. They can potentially be relinked at any point, so cannot rely on the host relevancy flags.

#Synced-CL 52001925

diff --git a/Engine/Plugins/Experimental/UAF/UAFAnimGraph/Source/UAFAnimGraph/Private/Traits/InputValueTrait.cpp b/Engine/Plugins/Experimental/UAF/UAFAnimGraph/Source/UAFAnimGraph/Private/Traits/InputValueTrait.cpp
--- a/Engine/Plugins/Experimental/UAF/UAFAnimGraph/Source/UAFAnimGraph/Private/Traits/InputValueTrait.cpp
+++ b/Engine/Plugins/Experimental/UAF/UAFAnimGraph/Source/UAFAnimGraph/Private/Traits/InputValueTrait.cpp
@@ -98,7 +98,9 @@ namespace UE::UAF
 					TSharedPtr<FAnimNextGraphInstance> GraphInstance = GraphReference->GetOrAllocate(Context, Binding);
 					if (GraphInstance.IsValid() && GraphInstance->GetGraphRootPtr().IsValid())
 					{
-						TraversalQueue.Push(GraphInstance->GetGraphRootPtr(), TraitState.AsNewlyRelevant(TraitState.IsNewlyRelevant() || !GraphInstance->HasUpdated()));
+						// External graph instances need to avoid re-initializing based around host relevancy
+						// Effectively they are in charge of their initial relevancy
+						TraversalQueue.Push(GraphInstance->GetGraphRootPtr(), TraitState.AsNewlyRelevant(!GraphInstance->HasUpdated()));
 						GraphInstance->MarkAsUpdated();
 					}
 				}
@@ -106,9 +108,7 @@ namespace UE::UAF
 				{
 					if (TraitPtr->IsValid())
 					{
-						FAnimNextGraphInstance& Owner = TraitPtr->GetNodeInstance()->GetOwner();
-						TraversalQueue.Push(*TraitPtr, TraitState.AsNewlyRelevant(TraitState.IsNewlyRelevant() || !Owner.HasUpdated()));
-						Owner.MarkAsUpdated();
+						TraversalQueue.Push(*TraitPtr, TraitState);
 					}
 				}
 			}
