UAF Root Motion Extraction fix
Move the same checks we do in the old code regarding 'Force Lock Root Bone' to allow root motion to still be extracted.

#Synced-CL 51964473

diff --git a/Engine/Plugins/Experimental/UAF/UAFAnimGraph/Source/UAFAnimGraph/Private/EvaluationVM/Tasks/DecompressionTools.cpp b/Engine/Plugins/Experimental/UAF/UAFAnimGraph/Source/UAFAnimGraph/Private/EvaluationVM/Tasks/DecompressionTools.cpp
--- a/Engine/Plugins/Experimental/UAF/UAFAnimGraph/Source/UAFAnimGraph/Private/EvaluationVM/Tasks/DecompressionTools.cpp
+++ b/Engine/Plugins/Experimental/UAF/UAFAnimGraph/Source/UAFAnimGraph/Private/EvaluationVM/Tasks/DecompressionTools.cpp
@@ -713,7 +713,12 @@ void FDecompressionTools::GetBonePose(const UAnimSequence* AnimSequence, const F
 #endif
 	}
 	const bool bTreatAnimAsAdditive = (AnimSequence->IsValidAdditive() && !bForceUseRawData); // Raw data is never additive
-	const FRootMotionReset RootMotionReset(AnimSequence->bEnableRootMotion, AnimSequence->RootMotionRootLock, AnimSequence->bForceRootLock, AnimSequence->ExtractRootTrackTransform(FAnimExtractContext(0.0), /*&RequiredBones*/nullptr), bTreatAnimAsAdditive);
+	const bool bApplyForceRootLock =
+#if WITH_EDITOR
+		!ExtractionContext.bIgnoreRootLock &&
+#endif
+		AnimSequence->bForceRootLock;
+	const FRootMotionReset RootMotionReset(AnimSequence->bEnableRootMotion, AnimSequence->RootMotionRootLock, bApplyForceRootLock, AnimSequence->ExtractRootTrackTransform(FAnimExtractContext(0.0), /*&RequiredBones*/nullptr), bTreatAnimAsAdditive);
 
 #if WITH_EDITOR
 	// Evaluate raw (source) curve and bone data
