UAF:
Move MaxGraphUpdateLoopCount setting from UAFEditorProjectSettings to UAFAnimGraphEditorProjectSettings
Change startup caching behavior to set the CVar to the project setting value instead of the other way around to avoid overriding project settings with CVar default values

#Synced-CL 52381678

diff --git a/Engine/Plugins/Experimental/UAF/UAF/Config/DefaultUAF.ini b/Engine/Plugins/Experimental/UAF/UAF/Config/DefaultUAF.ini
--- a/Engine/Plugins/Experimental/UAF/UAF/Config/DefaultUAF.ini
+++ b/Engine/Plugins/Experimental/UAF/UAF/Config/DefaultUAF.ini
@@ -317,4 +317,7 @@
 +ClassRedirects=(OldName="/Script/UAF.AnimNextSharedVariables",NewName="/Script/UAF.UAFSharedVariables")
 +ClassRedirects=(OldName="/Script/UAFUncookedOnly.AnimNextSharedVariablesEntry",NewName="/Script/UAFUncookedOnly.UAFSharedVariablesEntry")
 +ClassRedirects=(OldName="/Script/UAFUncookedOnly.AnimNextSharedVariableNode",NewName="/Script/UAFUncookedOnly.UAFSharedVariableNode")
-+ClassRedirects=(OldName="/Script/UAFEditor.AnimNextSharedVariablesFactory",NewName="/Script/UAFEditor.UAFSharedVariablesFactory")
\ No newline at end of file
++ClassRedirects=(OldName="/Script/UAFEditor.AnimNextSharedVariablesFactory",NewName="/Script/UAFEditor.UAFSharedVariablesFactory")
+
+; Move MaxGraphUpdateLoopCount from UAF to UAFAnimGraph
++PropertyRedirects=(OldName="/Script/UAFEditor.UAFEditorProjectSettings.MaxGraphUpdateLoopCount",NewName="/Script/UAFAnimGraphEditor.UAFAnimGraphEditorProjectSettings.MaxGraphUpdateLoopCount")
\ No newline at end of file
diff --git a/Engine/Plugins/Experimental/UAF/UAF/Source/UAFEditor/Private/Settings/UAFEditorProjectSettings.cpp b/Engine/Plugins/Experimental/UAF/UAF/Source/UAFEditor/Private/Settings/UAFEditorProjectSettings.cpp
--- a/Engine/Plugins/Experimental/UAF/UAF/Source/UAFEditor/Private/Settings/UAFEditorProjectSettings.cpp
+++ b/Engine/Plugins/Experimental/UAF/UAF/Source/UAFEditor/Private/Settings/UAFEditorProjectSettings.cpp
@@ -8,18 +8,3 @@ UUAFEditorProjectSettings::UUAFEditorProjectSettings()
 {
 	UAFBrowserCategoryConfiguration = "/UAFAnimGraph/Internal/TABC_UAFBrowserDefaultConfig.TABC_UAFBrowserDefaultConfig";
 }
-
-void UUAFEditorProjectSettings::PostInitProperties()
-{
-	Super::PostInitProperties();
-	ImportConsoleVariableValues();
-}
-
-void UUAFEditorProjectSettings::PostEditChangeProperty(FPropertyChangedEvent& PropertyChangedEvent)
-{
-	Super::PostEditChangeProperty(PropertyChangedEvent);
-	if (PropertyChangedEvent.Property)
-	{
-		ExportValuesToConsoleVariables(PropertyChangedEvent.Property);
-	}
-}
diff --git a/Engine/Plugins/Experimental/UAF/UAF/Source/UAFEditor/Private/Settings/UAFEditorProjectSettings.h b/Engine/Plugins/Experimental/UAF/UAF/Source/UAFEditor/Private/Settings/UAFEditorProjectSettings.h
--- a/Engine/Plugins/Experimental/UAF/UAF/Source/UAFEditor/Private/Settings/UAFEditorProjectSettings.h
+++ b/Engine/Plugins/Experimental/UAF/UAF/Source/UAFEditor/Private/Settings/UAFEditorProjectSettings.h
@@ -15,17 +15,8 @@ class UUAFEditorProjectSettings : public UDeveloperSettings
 
 	UUAFEditorProjectSettings();
 
-	UAFEDITOR_API virtual void PostInitProperties() override;
-	UAFEDITOR_API virtual void PostEditChangeProperty(FPropertyChangedEvent& PropertyChangedEvent) override;
 public:
-
 	/** The configuration used the UAF browser for setting up categories */
 	UPROPERTY(config, EditAnywhere, DisplayName = "UAF Browser Configuration", Category = UAF, meta = (AllowedClasses = "/Script/UserAssetTagsEditor.TaggedAssetBrowserConfiguration"))
 	FSoftObjectPath UAFBrowserCategoryConfiguration;
-	
-	/** The maximum amount of update loops per frame before a graph is considered to be stuck in an endless loop and breaks out of the update 
-	 *  This setting is linked with the CVar UAF.MaxGraphUpdateLoopCount 
-	 */
-	UPROPERTY(EditAnywhere, Category="UAF|Graph", Config, meta=(ConsoleVariable="UAF.MaxGraphUpdateLoopCount"))
-	int32 MaxGraphUpdateLoopCount = 1000;
 };
\ No newline at end of file
diff --git a/Engine/Plugins/Experimental/UAF/UAFAnimGraph/Source/UAFAnimGraphEditor/Private/Settings/UAFAnimGraphEditorProjectSettings.cpp b/Engine/Plugins/Experimental/UAF/UAFAnimGraph/Source/UAFAnimGraphEditor/Private/Settings/UAFAnimGraphEditorProjectSettings.cpp
new file mode 100644
--- /dev/null
+++ b/Engine/Plugins/Experimental/UAF/UAFAnimGraph/Source/UAFAnimGraphEditor/Private/Settings/UAFAnimGraphEditorProjectSettings.cpp
@@ -0,0 +1,37 @@
+// Copyright Epic Games, Inc. All Rights Reserved.
+
+#include "UAFAnimGraphEditorProjectSettings.h"
+
+#include "EngineLogs.h"
+#include "HAL/IConsoleManager.h"
+
+#include UE_INLINE_GENERATED_CPP_BY_NAME(UAFAnimGraphEditorProjectSettings)
+
+void UUAFAnimGraphEditorProjectSettings::PostInitProperties()
+{
+	Super::PostInitProperties();
+
+#if WITH_EDITOR
+	if (IsTemplate())
+	{
+		// Sets the CVar for MaxGraphUpdateLoopCount to the value of the project setting 
+		if (IConsoleVariable* CVar = IConsoleManager::Get().FindConsoleVariable(TEXT("UAF.MaxGraphUpdateLoopCount")))
+		{
+			CVar->Set(MaxGraphUpdateLoopCount, ECVF_SetByProjectSetting);
+		}
+		else
+		{
+			UE_LOG(LogAnimation, Error, TEXT("Could not find console variable UAF.MaxGraphUpdateLoopCount"));
+		}
+	}
+#endif
+}
+
+void UUAFAnimGraphEditorProjectSettings::PostEditChangeProperty(FPropertyChangedEvent& PropertyChangedEvent)
+{
+	Super::PostEditChangeProperty(PropertyChangedEvent);
+	if (PropertyChangedEvent.Property)
+	{
+		ExportValuesToConsoleVariables(PropertyChangedEvent.Property);
+	}
+}
diff --git a/Engine/Plugins/Experimental/UAF/UAFAnimGraph/Source/UAFAnimGraphEditor/Private/Settings/UAFAnimGraphEditorProjectSettings.h b/Engine/Plugins/Experimental/UAF/UAFAnimGraph/Source/UAFAnimGraphEditor/Private/Settings/UAFAnimGraphEditorProjectSettings.h
new file mode 100644
--- /dev/null
+++ b/Engine/Plugins/Experimental/UAF/UAFAnimGraph/Source/UAFAnimGraphEditor/Private/Settings/UAFAnimGraphEditorProjectSettings.h
@@ -0,0 +1,27 @@
+// Copyright Epic Games, Inc. All Rights Reserved.
+
+#pragma once
+
+#include "CoreMinimal.h"
+#include "Engine/DeveloperSettings.h"
+
+#include "UAFAnimGraphEditorProjectSettings.generated.h"
+
+/**
+ * UAF Anim Graph Editor Settings that are shared across a project.
+ */
+UCLASS(MinimalAPI, Config = Editor, defaultconfig, DisplayName = "UAF Anim Graph Settings")
+class UUAFAnimGraphEditorProjectSettings : public UDeveloperSettings
+{
+	GENERATED_BODY()
+	
+public:
+	UAFANIMGRAPHEDITOR_API virtual void PostInitProperties() override;
+	UAFANIMGRAPHEDITOR_API virtual void PostEditChangeProperty(FPropertyChangedEvent& PropertyChangedEvent) override;
+	
+public:
+	/** The maximum amount of update loops per frame before a graph is considered to be stuck in an endless loop and breaks out of the update 
+	 *  This setting is linked with the CVar UAF.MaxGraphUpdateLoopCount  */
+	UPROPERTY(Config, EditAnywhere, Category="UAF|Graph", meta=(ConsoleVariable="UAF.MaxGraphUpdateLoopCount"))
+	int32 MaxGraphUpdateLoopCount = 1000;
+};
diff --git a/Engine/Plugins/Experimental/UAF/UAFAnimGraph/Source/UAFAnimGraphEditor/UAFAnimGraphEditor.Build.cs b/Engine/Plugins/Experimental/UAF/UAFAnimGraph/Source/UAFAnimGraphEditor/UAFAnimGraphEditor.Build.cs
--- a/Engine/Plugins/Experimental/UAF/UAFAnimGraph/Source/UAFAnimGraphEditor/UAFAnimGraphEditor.Build.cs
+++ b/Engine/Plugins/Experimental/UAF/UAFAnimGraph/Source/UAFAnimGraphEditor/UAFAnimGraphEditor.Build.cs
@@ -49,6 +49,7 @@ namespace UnrealBuildTool.Rules
 					"BlueprintGraph",
 					"ToolWidgets",
 					"Kismet",
+					"DeveloperSettings",
 				}
 			);
 		}
