UAF: Add Inertialization/Dead Blending support for UAF Montages

#Synced-CL 52535972

diff --git a/Engine/Plugins/Experimental/UAF/UAFAnimGraph/Source/UAFAnimGraph/Internal/Injection/UAFMontageComponent.h b/Engine/Plugins/Experimental/UAF/UAFAnimGraph/Source/UAFAnimGraph/Internal/Injection/UAFMontageComponent.h
--- a/Engine/Plugins/Experimental/UAF/UAFAnimGraph/Source/UAFAnimGraph/Internal/Injection/UAFMontageComponent.h
+++ b/Engine/Plugins/Experimental/UAF/UAFAnimGraph/Source/UAFAnimGraph/Internal/Injection/UAFMontageComponent.h
@@ -96,6 +96,9 @@ private:
 	// Retrieves all active montage animation notifies for a specific slot name
 	// returns true if any notifies for the given slot name exist
 	bool GetActiveNotifiesForSlot(FName SlotName, TArray<FAnimNotifyEventReference>& OutNotifies) const;
+	
+	bool HasActiveInertializationRequests() const;
+	const FInertializationRequest* GetInertializationRequestForGroupName(FName SlotGroupName) const;
 
 private:
 	// All active and weighted montages, copied from the Anim Instance
@@ -105,6 +108,9 @@ private:
 	// Active animation notifies per slot, used to have them easily accessible in montage traits 
 	UPROPERTY(Transient)
 	TMap<FName, FAnimNotifyArray> ActiveSlotNotifies;
+	
+	UPROPERTY(Transient)
+	TMap<FName, FInertializationRequest> ActiveInertializationRequests;
 
 	// The tick function that is used to update montages on the game thread
 	FUAFUpdateMontageTickFunction MontageTickFunction;
diff --git a/Engine/Plugins/Experimental/UAF/UAFAnimGraph/Source/UAFAnimGraph/Private/Injection/MontageTrait.cpp b/Engine/Plugins/Experimental/UAF/UAFAnimGraph/Source/UAFAnimGraph/Private/Injection/MontageTrait.cpp
--- a/Engine/Plugins/Experimental/UAF/UAFAnimGraph/Source/UAFAnimGraph/Private/Injection/MontageTrait.cpp
+++ b/Engine/Plugins/Experimental/UAF/UAFAnimGraph/Source/UAFAnimGraph/Private/Injection/MontageTrait.cpp
@@ -21,6 +21,8 @@
 #include "EvaluationVM/Tasks/BlendKeyframesPerBone.h"
 #include "EvaluationVM/Tasks/ApplyAdditiveKeyframePerBone.h"
 #include "UAFLogging.h"
+#include "Engine/SkeletalMesh.h"
+#include "Traits/Inertialization.h"
 
 namespace UE::UAF
 {
@@ -92,12 +94,7 @@ namespace UE::UAF
 			UAF_TRAIT_LOG(Warning, TEXT("FMontageTrait::PreUpdate: Could not update montages - There is no valid MontageComponent."));
 			return;
 		}
-
-		if (MontageComponent->MontageEvalData.IsEmpty())
-		{
-			return;
-		}
-
+		
 		const FSharedData* SharedData = Binding.GetSharedData<FSharedData>();
 		const FName CurrentSlotName = SharedData->GetSlotName(Binding);
 		if (CurrentSlotName.IsNone())
@@ -105,7 +102,31 @@ namespace UE::UAF
 			UAF_TRAIT_LOG(Warning, TEXT("FMontageTrait::PreUpdate: Could not update montages - Trait has not valid slot name."));
 			return;
 		}
+		
+		// Raise Inertialization events if applicable
+		if (MontageComponent->HasActiveInertializationRequests())
+		{
+			const USkeleton* Skeleton =	Context.GetBindingObject().IsValid() && Context.GetBindingObject()->GetSkeletalMeshAsset() 
+										? Context.GetBindingObject()->GetSkeletalMeshAsset()->GetSkeleton() 
+										: nullptr;
+			
+			if (Skeleton)
+			{
+				const FName SlotGroupName = Skeleton->GetSlotGroupName(CurrentSlotName);
+				if (const ::FInertializationRequest* ActiveRequest = MontageComponent->GetInertializationRequestForGroupName(SlotGroupName))
+				{
+					TSharedPtr<FAnimNextInertializationRequestEvent> InertializationRequestEvent = MakeTraitEvent<FAnimNextInertializationRequestEvent>();
+					InertializationRequestEvent->Request.BlendTime = ActiveRequest->Duration;
+					Context.RaiseOutputTraitEvent(InertializationRequestEvent);
+				}
+			}
+		}
 
+		if (MontageComponent->MontageEvalData.IsEmpty())
+		{
+			return;
+		}
+		
 		InstanceData->bUsesBlendProfiles = false;
 		for (const FUAFMontageEvalData& EvalData : MontageComponent->MontageEvalData)
 		{
diff --git a/Engine/Plugins/Experimental/UAF/UAFAnimGraph/Source/UAFAnimGraph/Private/Injection/UAFMontageComponent.cpp b/Engine/Plugins/Experimental/UAF/UAFAnimGraph/Source/UAFAnimGraph/Private/Injection/UAFMontageComponent.cpp
--- a/Engine/Plugins/Experimental/UAF/UAFAnimGraph/Source/UAFAnimGraph/Private/Injection/UAFMontageComponent.cpp
+++ b/Engine/Plugins/Experimental/UAF/UAFAnimGraph/Source/UAFAnimGraph/Private/Injection/UAFMontageComponent.cpp
@@ -52,8 +52,12 @@ void FUAFUpdateMontageTickFunction::ExecuteTick(float DeltaTime, enum ELevelTick
 	MontageComponent->ActiveSlotNotifies.Reset();
 	MontageComponent->ActiveSlotNotifies = WeakAnimInstance->NotifyQueue.UnfilteredMontageAnimNotifies;
 
+	// Copy any active inertialization requests from the anim instance and then clear them manually since UpdateMontageEvaluationData won't be called
+	MontageComponent->ActiveInertializationRequests = WeakAnimInstance->GetActiveInertializationRequests();
+	WeakAnimInstance->ClearActiveInertializationRequests();
+	
 	MontageComponent->MontageEvalData.Reset();
-	for (FAnimMontageInstance* MontageInstance : WeakAnimInstance->MontageInstances)
+	for (const FAnimMontageInstance* MontageInstance : WeakAnimInstance->MontageInstances)
 	{
 		// although montage can advance with 0.f weight, it is fine to filter by weight here
 		// because we don't want to evaluate them if 0 weight
@@ -161,6 +165,16 @@ TObjectPtr<const UAnimMontage> FUAFMontageComponent::GetActiveMontageForSlot(FNa
 	return nullptr;
 }
 
+bool FUAFMontageComponent::HasActiveInertializationRequests() const
+{
+	return !ActiveInertializationRequests.IsEmpty();
+}
+
+const FInertializationRequest* FUAFMontageComponent::GetInertializationRequestForGroupName(FName SlotGroupName) const
+{
+	return ActiveInertializationRequests.Find(SlotGroupName);
+}
+
 bool FUAFMontageComponent::GetActiveNotifiesForSlot(FName SlotName, TArray<FAnimNotifyEventReference>& OutNotifies) const
 {
 	const FAnimNotifyArray* FoundEntry = ActiveSlotNotifies.Find(SlotName);
diff --git a/Engine/Source/Runtime/Engine/Classes/Animation/AnimInstance.h b/Engine/Source/Runtime/Engine/Classes/Animation/AnimInstance.h
--- a/Engine/Source/Runtime/Engine/Classes/Animation/AnimInstance.h
+++ b/Engine/Source/Runtime/Engine/Classes/Animation/AnimInstance.h
@@ -1012,6 +1012,12 @@ public:
 	/**  Requests an inertial blend during the next anim graph update. Requires your anim graph to have a slot node belonging to the specified group name */
 	UFUNCTION(BlueprintCallable, Category = "Animation|Inertial Blending")
 	ENGINE_API void RequestSlotGroupInertialization(FName InSlotGroupName, float Duration, const UBlendProfile* BlendProfile = nullptr);
+	
+	/**  Returns active inertialization requests on the anim instance (does not forward the call to the proxy), used for legacy montage support in UAF */
+	ENGINE_API const TMap<FName, FInertializationRequest>& GetActiveInertializationRequests() const;
+	
+	/**  Clears all active inertialization requests on the anim instance (does not forward the call to the proxy), used for legacy montage support in UAF */
+	ENGINE_API void ClearActiveInertializationRequests();
 
 	/** Queue a Montage BlendingOut Event to be triggered. */
 	ENGINE_API void QueueMontageBlendingOutEvent(const FQueuedMontageBlendingOutEvent& MontageBlendingOutEvent);
diff --git a/Engine/Source/Runtime/Engine/Private/Animation/AnimInstance.cpp b/Engine/Source/Runtime/Engine/Private/Animation/AnimInstance.cpp
--- a/Engine/Source/Runtime/Engine/Private/Animation/AnimInstance.cpp
+++ b/Engine/Source/Runtime/Engine/Private/Animation/AnimInstance.cpp
@@ -2351,6 +2351,16 @@ void UAnimInstance::RequestSlotGroupInertialization(FName InSlotGroupName, float
 	GetProxyOnAnyThread<FAnimInstanceProxy>().GetSlotGroupInertializationRequestDataMap().FindOrAdd(InSlotGroupName) = Request;
 }
 
+const TMap<FName, FInertializationRequest>& UAnimInstance::GetActiveInertializationRequests() const
+{
+	return SlotGroupInertializationRequestDataMap;
+}
+
+void UAnimInstance::ClearActiveInertializationRequests()
+{
+	SlotGroupInertializationRequestDataMap.Reset();
+}
+
 void UAnimInstance::ConditionalFlushCompletedMontages()
 {
 	if (MontageInstances.IsEmpty())
